L10N: Updated the WordCount.ini to account for the new locations of the PropertyNames and ToolTips targets.

diff --git a/Engine/Config/Localization/WordCount.ini b/Engine/Config/Localization/WordCount.ini
--- a/Engine/Config/Localization/WordCount.ini
+++ b/Engine/Config/Localization/WordCount.ini
@@ -27,7 +27,7 @@ WordCountReportName=Editor.csv
 [GatherTextStep2]
 CommandletClass=GenerateTextLocalizationReport
 bWordCountReport=true
-SourcePath=../../../Engine/Content/Localization/Editor
+SourcePath=../../../Engine/Content/Localization/ToolTips
 ManifestName=ToolTips.manifest
 CulturesToGenerate=ja
 CulturesToGenerate=ko
@@ -38,7 +38,7 @@ WordCountReportName=ToolTips.csv
 [GatherTextStep3]
 CommandletClass=GenerateTextLocalizationReport
 bWordCountReport=true
-SourcePath=../../../Engine/Content/Localization/Editor
+SourcePath=../../../Engine/Content/Localization/PropertyNames
 ManifestName=PropertyNames.manifest
 CulturesToGenerate=ja
 CulturesToGenerate=ko
