Updated translations. Imported new spanish translations

diff --git a/Engine/Config/Localization/Editor.ini b/Engine/Config/Localization/Editor.ini
--- a/Engine/Config/Localization/Editor.ini
+++ b/Engine/Config/Localization/Editor.ini
@@ -9,6 +9,7 @@ CulturesToGenerate=en
 CulturesToGenerate=ja
 CulturesToGenerate=ko
 CulturesToGenerate=zh-CN
+CulturesToGenerate=sp
 
 ;Gather text from source code
 [GatherTextStep0]
diff --git a/Engine/Config/Localization/Engine.ini b/Engine/Config/Localization/Engine.ini
--- a/Engine/Config/Localization/Engine.ini
+++ b/Engine/Config/Localization/Engine.ini
@@ -8,6 +8,7 @@ CulturesToGenerate=en
 CulturesToGenerate=ja
 CulturesToGenerate=ko
 CulturesToGenerate=zh-CN
+CulturesToGenerate=sp
 
 ;Gather text from source code
 [GatherTextStep0]
diff --git a/Engine/Config/Localization/PropertyNames.ini b/Engine/Config/Localization/PropertyNames.ini
--- a/Engine/Config/Localization/PropertyNames.ini
+++ b/Engine/Config/Localization/PropertyNames.ini
@@ -8,6 +8,7 @@ CulturesToGenerate=en
 CulturesToGenerate=ja
 CulturesToGenerate=ko
 CulturesToGenerate=zh-CN
+CulturesToGenerate=sp
 
 ;Gather text from metadata
 [GatherTextStep0]
diff --git a/Engine/Config/Localization/ToolTips.ini b/Engine/Config/Localization/ToolTips.ini
--- a/Engine/Config/Localization/ToolTips.ini
+++ b/Engine/Config/Localization/ToolTips.ini
@@ -8,6 +8,7 @@ CulturesToGenerate=en
 CulturesToGenerate=ja
 CulturesToGenerate=ko
 CulturesToGenerate=zh-CN
+CulturesToGenerate=sp
 
 ;Gather text from metadata
 [GatherTextStep0]
