Updated the text for the library tutorials. Included UnrealEngineLauncherServices module as part of the launcher localization and excluded it from the engine/editor localization.

diff --git a/Engine/Config/Localization/Editor.ini b/Engine/Config/Localization/Editor.ini
--- a/Engine/Config/Localization/Editor.ini
+++ b/Engine/Config/Localization/Editor.ini
@@ -29,6 +29,7 @@ IncludePaths=*/Engine/Source/Runtime/*
 IncludePaths=*/Engine/Source/Developer/*
 ExcludePaths=*/Engine/Config/Localization/*
 ExcludePaths=*/Engine/Source/Developer/NoRedist/CommunityPortalServices/*
+ExcludePaths=*/Engine/Source/Developer/NoRedist/UnrealEngineLauncherServices/*
 ExcludePaths=*/Engine/Source/Developer/NoRedist/BuildPatchServices/*
 InputKeys=Category
 OutputNamespaces=UObjectCategories
diff --git a/Engine/Config/Localization/Engine.ini b/Engine/Config/Localization/Engine.ini
--- a/Engine/Config/Localization/Engine.ini
+++ b/Engine/Config/Localization/Engine.ini
@@ -19,6 +19,7 @@ IncludePaths=Engine/Source/Developer
 IncludePaths=Engine/Config/
 ExcludePaths=*/Engine/Config/Localization/*
 ExcludePaths=*/Engine/Source/Developer/NoRedist/CommunityPortalServices/*
+ExcludePaths=*/Engine/Source/Developer/NoRedist/UnrealEngineLauncherServices/*
 ExcludePaths=*/Engine/Source/Developer/NoRedist/BuildPatchServices/*
 SourceFileSearchFilters=*.cpp
 SourceFileSearchFilters=*.h
