Renamed Spanish localization data folder from "sp" to "es". Corrected values in localization data generation configuration files similarly.

diff --git a/Engine/Config/Localization/Editor.ini b/Engine/Config/Localization/Editor.ini
--- a/Engine/Config/Localization/Editor.ini
+++ b/Engine/Config/Localization/Editor.ini
@@ -9,7 +9,7 @@ CulturesToGenerate=en
 CulturesToGenerate=ja
 CulturesToGenerate=ko
 CulturesToGenerate=zh-CN
-CulturesToGenerate=sp
+CulturesToGenerate=es
 
 ;Gather text from source code
 [GatherTextStep0]
@@ -59,14 +59,23 @@ CommandletClass=GenerateGatherManifest
 CommandletClass=GenerateGatherArchive
 bPurgeOldEmptyEntries=true
 
-;Write Localized Text Resource
+;Import localized PO files
 [GatherTextStep5]
+CommandletClass=InternationalizationExport
+bImportLoc=true
+ManifestName=Editor.manifest
+ArchiveName=Editor.archive
+PortableObjectName=Editor.po
+
+;Write Localized Text Resource
+[GatherTextStep6]
 CommandletClass=GenerateTextLocalizationResource
 ResourceName=Editor.locres
 
-;Write Text Localization Report
-[GatherTextStep6]
-CommandletClass=GenerateTextLocalizationReport
-DestinationPath=\\epicgames.net\root\UE3\Localization\WordCounts
-bConflictReport=true
-ConflictReportName=EditorConflicts.txt
\ No newline at end of file
+;Import localized PO files
+[GatherTextStep7]
+CommandletClass=InternationalizationExport
+bExportLoc=true
+ManifestName=Editor.manifest
+ArchiveName=Editor.archive
+PortableObjectName=Editor.po
\ No newline at end of file
diff --git a/Engine/Config/Localization/Engine.ini b/Engine/Config/Localization/Engine.ini
--- a/Engine/Config/Localization/Engine.ini
+++ b/Engine/Config/Localization/Engine.ini
@@ -8,7 +8,7 @@ CulturesToGenerate=en
 CulturesToGenerate=ja
 CulturesToGenerate=ko
 CulturesToGenerate=zh-CN
-CulturesToGenerate=sp
+CulturesToGenerate=es
 
 ;Gather text from source code
 [GatherTextStep0]
diff --git a/Engine/Config/Localization/PropertyNames.ini b/Engine/Config/Localization/PropertyNames.ini
--- a/Engine/Config/Localization/PropertyNames.ini
+++ b/Engine/Config/Localization/PropertyNames.ini
@@ -8,7 +8,7 @@ CulturesToGenerate=en
 CulturesToGenerate=ja
 CulturesToGenerate=ko
 CulturesToGenerate=zh-CN
-CulturesToGenerate=sp
+CulturesToGenerate=es
 
 ;Gather text from metadata
 [GatherTextStep0]
diff --git a/Engine/Config/Localization/ToolTips.ini b/Engine/Config/Localization/ToolTips.ini
--- a/Engine/Config/Localization/ToolTips.ini
+++ b/Engine/Config/Localization/ToolTips.ini
@@ -8,7 +8,7 @@ CulturesToGenerate=en
 CulturesToGenerate=ja
 CulturesToGenerate=ko
 CulturesToGenerate=zh-CN
-CulturesToGenerate=sp
+CulturesToGenerate=es
 
 ;Gather text from metadata
 [GatherTextStep0]
