Updated localization data gathering and generation configuration files to be consistent in terms of cultures generated and directories used.

diff --git a/Engine/Config/Localization/PortableObjectExport.ini b/Engine/Config/Localization/PortableObjectExport.ini
--- a/Engine/Config/Localization/PortableObjectExport.ini
+++ b/Engine/Config/Localization/PortableObjectExport.ini
@@ -1,5 +1,6 @@
 [CommonSettings]
 CulturesToGenerate=en
+CulturesToGenerate=es
 CulturesToGenerate=ja
 CulturesToGenerate=ko
 CulturesToGenerate=zh-CN
@@ -25,8 +26,8 @@ PortableObjectName=Editor.po
 [GatherTextStep2]
 CommandletClass=InternationalizationExport
 bExportLoc=true
-SourcePath=../../../Engine/Content/Localization/Editor
-DestinationPath=../../../Engine/Content/Localization/Editor
+SourcePath=../../../Engine/Content/Localization/ToolTips
+DestinationPath=../../../Engine/Content/Localization/ToolTips
 ManifestName=ToolTips.manifest
 ArchiveName=ToolTips.archive
 PortableObjectName=ToolTips.po
@@ -34,8 +35,8 @@ PortableObjectName=ToolTips.po
 [GatherTextStep3]
 CommandletClass=InternationalizationExport
 bExportLoc=true
-SourcePath=../../../Engine/Content/Localization/Editor
-DestinationPath=../../../Engine/Content/Localization/Editor
+SourcePath=../../../Engine/Content/Localization/PropertyNames
+DestinationPath=../../../Engine/Content/Localization/PropertyNames
 ManifestName=PropertyNames.manifest
 ArchiveName=PropertyNames.archive
 PortableObjectName=PropertyNames.po
\ No newline at end of file
diff --git a/Engine/Config/Localization/PortableObjectImport.ini b/Engine/Config/Localization/PortableObjectImport.ini
--- a/Engine/Config/Localization/PortableObjectImport.ini
+++ b/Engine/Config/Localization/PortableObjectImport.ini
@@ -1,5 +1,6 @@
 [CommonSettings]
 CulturesToGenerate=en
+CulturesToGenerate=es
 CulturesToGenerate=ja
 CulturesToGenerate=ko
 CulturesToGenerate=zh-CN
diff --git a/Engine/Config/Localization/ResourceFileGen.ini b/Engine/Config/Localization/ResourceFileGen.ini
--- a/Engine/Config/Localization/ResourceFileGen.ini
+++ b/Engine/Config/Localization/ResourceFileGen.ini
@@ -5,6 +5,7 @@ SourcePath=../../../Engine/Content/Localization/Engine
 DestinationPath=../../../Engine/Content/Localization/Engine
 ManifestName=Engine.manifest
 CulturesToGenerate=en
+CulturesToGenerate=es
 CulturesToGenerate=ja
 CulturesToGenerate=ko
 CulturesToGenerate=zh-CN
@@ -17,6 +18,7 @@ SourcePath=../../../Engine/Content/Localization/Editor
 DestinationPath=../../../Engine/Content/Localization/Editor
 ManifestName=Editor.manifest
 CulturesToGenerate=en
+CulturesToGenerate=es
 CulturesToGenerate=ja
 CulturesToGenerate=ko
 CulturesToGenerate=zh-CN
@@ -29,6 +31,7 @@ SourcePath=../../../Engine/Content/Localization/Editor/ToolTips
 DestinationPath=../../../Engine/Content/Localization/Editor/ToolTips
 ManifestName=ToolTips.manifest
 CulturesToGenerate=en
+CulturesToGenerate=es
 CulturesToGenerate=ja
 CulturesToGenerate=ko
 CulturesToGenerate=zh-CN
@@ -41,6 +44,7 @@ SourcePath=../../../Engine/Content/Localization/Editor/PropertyNames
 DestinationPath=../../../Engine/Content/Localization/Editor/PropertyNames
 ManifestName=PropertyNames.manifest
 CulturesToGenerate=en
+CulturesToGenerate=es
 CulturesToGenerate=ja
 CulturesToGenerate=ko
 CulturesToGenerate=zh-CN
diff --git a/Engine/Config/Localization/WordCount.ini b/Engine/Config/Localization/WordCount.ini
--- a/Engine/Config/Localization/WordCount.ini
+++ b/Engine/Config/Localization/WordCount.ini
@@ -7,6 +7,7 @@ CommandletClass=GenerateTextLocalizationReport
 bWordCountReport=true
 SourcePath=../../../Engine/Content/Localization/Engine
 ManifestName=Engine.manifest
+CulturesToGenerate=es
 CulturesToGenerate=ja
 CulturesToGenerate=ko
 CulturesToGenerate=zh-CN
@@ -18,6 +19,7 @@ CommandletClass=GenerateTextLocalizationReport
 bWordCountReport=true
 SourcePath=../../../Engine/Content/Localization/Editor
 ManifestName=Editor.manifest
+CulturesToGenerate=es
 CulturesToGenerate=ja
 CulturesToGenerate=ko
 CulturesToGenerate=zh-CN
@@ -29,6 +31,7 @@ CommandletClass=GenerateTextLocalizationReport
 bWordCountReport=true
 SourcePath=../../../Engine/Content/Localization/ToolTips
 ManifestName=ToolTips.manifest
+CulturesToGenerate=es
 CulturesToGenerate=ja
 CulturesToGenerate=ko
 CulturesToGenerate=zh-CN
@@ -40,6 +43,7 @@ CommandletClass=GenerateTextLocalizationReport
 bWordCountReport=true
 SourcePath=../../../Engine/Content/Localization/PropertyNames
 ManifestName=PropertyNames.manifest
+CulturesToGenerate=es
 CulturesToGenerate=ja
 CulturesToGenerate=ko
 CulturesToGenerate=zh-CN
