Updated localization data generation configuration files to add Brazilian Portuguese.

diff --git a/Engine/Config/Localization/Editor.ini b/Engine/Config/Localization/Editor.ini
--- a/Engine/Config/Localization/Editor.ini
+++ b/Engine/Config/Localization/Editor.ini
@@ -10,6 +10,7 @@ CulturesToGenerate=en
 CulturesToGenerate=es
 CulturesToGenerate=ja
 CulturesToGenerate=ko
+CulturesToGenerate=pt-BR
 CulturesToGenerate=zh-Hans-CN
 
 ;Gather text from source code
diff --git a/Engine/Config/Localization/Engine.ini b/Engine/Config/Localization/Engine.ini
--- a/Engine/Config/Localization/Engine.ini
+++ b/Engine/Config/Localization/Engine.ini
@@ -9,6 +9,7 @@ CulturesToGenerate=en
 CulturesToGenerate=es
 CulturesToGenerate=ja
 CulturesToGenerate=ko
+CulturesToGenerate=pt-BR
 CulturesToGenerate=zh-Hans-CN
 
 ;Gather text from source code
diff --git a/Engine/Config/Localization/PortableObjectExport.ini b/Engine/Config/Localization/PortableObjectExport.ini
--- a/Engine/Config/Localization/PortableObjectExport.ini
+++ b/Engine/Config/Localization/PortableObjectExport.ini
@@ -4,6 +4,7 @@ CulturesToGenerate=en
 CulturesToGenerate=es
 CulturesToGenerate=ja
 CulturesToGenerate=ko
+CulturesToGenerate=pt-BR
 CulturesToGenerate=zh-Hans-CN
 
 [GatherTextStep0]
diff --git a/Engine/Config/Localization/PortableObjectImport.ini b/Engine/Config/Localization/PortableObjectImport.ini
--- a/Engine/Config/Localization/PortableObjectImport.ini
+++ b/Engine/Config/Localization/PortableObjectImport.ini
@@ -4,6 +4,7 @@ CulturesToGenerate=en
 CulturesToGenerate=es
 CulturesToGenerate=ja
 CulturesToGenerate=ko
+CulturesToGenerate=pt-BR
 CulturesToGenerate=zh-Hans-CN
 
 [GatherTextStep0]
diff --git a/Engine/Config/Localization/PropertyNames.ini b/Engine/Config/Localization/PropertyNames.ini
--- a/Engine/Config/Localization/PropertyNames.ini
+++ b/Engine/Config/Localization/PropertyNames.ini
@@ -8,6 +8,7 @@ CulturesToGenerate=es
 CulturesToGenerate=en
 CulturesToGenerate=ja
 CulturesToGenerate=ko
+CulturesToGenerate=pt-BR
 CulturesToGenerate=zh-Hans-CN
 
 ;Gather text from metadata
@@ -34,11 +35,4 @@ bPurgeOldEmptyEntries=true
 ;Write Localized Text Resource
 [GatherTextStep3]
 CommandletClass=GenerateTextLocalizationResource
-ResourceName=PropertyNames.locres
-
-;Write Text Localization Report
-[GatherTextStep4]
-CommandletClass=GenerateTextLocalizationReport
-DestinationPath=\\epicgames.net\root\UE3\Localization\WordCounts
-bConflictReport=true
-ConflictReportName=PropertyNameConflicts.txt
\ No newline at end of file
+ResourceName=PropertyNames.locres
\ No newline at end of file
diff --git a/Engine/Config/Localization/ResourceFileGen.ini b/Engine/Config/Localization/ResourceFileGen.ini
--- a/Engine/Config/Localization/ResourceFileGen.ini
+++ b/Engine/Config/Localization/ResourceFileGen.ini
@@ -3,6 +3,7 @@ CulturesToGenerate=en
 CulturesToGenerate=es
 CulturesToGenerate=ja
 CulturesToGenerate=ko
+CulturesToGenerate=pt-BR
 CulturesToGenerate=zh-Hans-CN
 
 ;Write Localized Text Resource for Engine
diff --git a/Engine/Config/Localization/ToolTips.ini b/Engine/Config/Localization/ToolTips.ini
--- a/Engine/Config/Localization/ToolTips.ini
+++ b/Engine/Config/Localization/ToolTips.ini
@@ -8,6 +8,7 @@ CulturesToGenerate=es
 CulturesToGenerate=en
 CulturesToGenerate=ja
 CulturesToGenerate=ko
+CulturesToGenerate=pt-BR
 CulturesToGenerate=zh-Hans-CN
 
 ;Gather text from metadata
@@ -34,11 +35,4 @@ bPurgeOldEmptyEntries=true
 ;Write Localized Text Resource
 [GatherTextStep3]
 CommandletClass=GenerateTextLocalizationResource
-ResourceName=ToolTips.locres
-
-;Write Text Localization Report
-[GatherTextStep4]
-CommandletClass=GenerateTextLocalizationReport
-DestinationPath=\\epicgames.net\root\UE3\Localization\WordCounts
-bConflictReport=true
-ConflictReportName=TooltipConflicts.txt
\ No newline at end of file
+ResourceName=ToolTips.locres
\ No newline at end of file
diff --git a/Engine/Config/Localization/WordCount.ini b/Engine/Config/Localization/WordCount.ini
--- a/Engine/Config/Localization/WordCount.ini
+++ b/Engine/Config/Localization/WordCount.ini
@@ -4,6 +4,7 @@ DestinationPath=\\epicgames.net\root\UE3\Localization\WordCounts
 CulturesToGenerate=es
 CulturesToGenerate=ja
 CulturesToGenerate=ko
+CulturesToGenerate=pt-BR
 CulturesToGenerate=zh-Hans-CN
 
 ;Write Text Localization Report for Engine
