Added ToolTips and PropertyNames *.po file importing to the PortableObjectImport.ini configuration file.

diff --git a/Engine/Config/Localization/PortableObjectImport.ini b/Engine/Config/Localization/PortableObjectImport.ini
--- a/Engine/Config/Localization/PortableObjectImport.ini
+++ b/Engine/Config/Localization/PortableObjectImport.ini
@@ -22,3 +22,19 @@ DestinationPath=./Content/Localization/Editor
 ManifestName=Editor.manifest
 ArchiveName=Editor.archive
 PortableObjectName=Editor.po
+
+[GatherTextStep2]
+CommandletClass=InternationalizationExport
+SourcePath=./Content/Localization/ToolTips
+DestinationPath=./Content/Localization/ToolTips
+ManifestName=ToolTips.manifest
+ArchiveName=ToolTips.archive
+PortableObjectName=ToolTips.po
+
+[GatherTextStep3]
+CommandletClass=InternationalizationExport
+SourcePath=./Content/Localization/PropertyNames
+DestinationPath=./Content/Localization/PropertyNames
+ManifestName=PropertyNames.manifest
+ArchiveName=PropertyNames.archive
+PortableObjectName=PropertyNames.po
\ No newline at end of file
