Corrected superficial contradictory comment in Editor localization data generation configuration file.

diff --git a/Engine/Config/Localization/Editor.ini b/Engine/Config/Localization/Editor.ini
--- a/Engine/Config/Localization/Editor.ini
+++ b/Engine/Config/Localization/Editor.ini
@@ -66,7 +66,7 @@ bImportLoc=true
 CommandletClass=GenerateTextLocalizationResource
 ResourceName=Editor.locres
 
-;Import localized PO files
+;Export to PO files
 [GatherTextStep7]
 CommandletClass=InternationalizationExport
 bExportLoc=true
\ No newline at end of file
