Added portable object import/export to text localization data generation configuration files for PropertyNames and ToolTips.

diff --git a/Engine/Config/Localization/PropertyNames.ini b/Engine/Config/Localization/PropertyNames.ini
--- a/Engine/Config/Localization/PropertyNames.ini
+++ b/Engine/Config/Localization/PropertyNames.ini
@@ -3,6 +3,7 @@ SourcePath=./Content/Localization/PropertyNames
 DestinationPath=./Content/Localization/PropertyNames
 ManifestName=PropertyNames.manifest
 ArchiveName=PropertyNames.archive
+PortableObjectName=PropertyNames.po
 SourceCulture=en
 CulturesToGenerate=es
 CulturesToGenerate=en
@@ -32,7 +33,17 @@ CommandletClass=GenerateGatherManifest
 CommandletClass=GenerateGatherArchive
 bPurgeOldEmptyEntries=true
 
-;Write Localized Text Resource
+;Import localized PO files
 [GatherTextStep3]
+CommandletClass=InternationalizationExport
+bImportLoc=true
+
+;Write Localized Text Resource
+[GatherTextStep4]
 CommandletClass=GenerateTextLocalizationResource
-ResourceName=PropertyNames.locres
\ No newline at end of file
+ResourceName=PropertyNames.locres
+
+;Export to PO files
+[GatherTextStep5]
+CommandletClass=InternationalizationExport
+bExportLoc=true
\ No newline at end of file
diff --git a/Engine/Config/Localization/ToolTips.ini b/Engine/Config/Localization/ToolTips.ini
--- a/Engine/Config/Localization/ToolTips.ini
+++ b/Engine/Config/Localization/ToolTips.ini
@@ -3,6 +3,7 @@ SourcePath=./Content/Localization/ToolTips
 DestinationPath=./Content/Localization/ToolTips
 ManifestName=ToolTips.manifest
 ArchiveName=ToolTips.archive
+PortableObjectName=ToolTips.po
 SourceCulture=en
 CulturesToGenerate=es
 CulturesToGenerate=en
@@ -32,7 +33,17 @@ CommandletClass=GenerateGatherManifest
 CommandletClass=GenerateGatherArchive
 bPurgeOldEmptyEntries=true
 
-;Write Localized Text Resource
+;Import localized PO files
 [GatherTextStep3]
+CommandletClass=InternationalizationExport
+bImportLoc=true
+
+;Write Localized Text Resource
+[GatherTextStep4]
 CommandletClass=GenerateTextLocalizationResource
-ResourceName=ToolTips.locres
\ No newline at end of file
+ResourceName=ToolTips.locres
+
+;Export to PO files
+[GatherTextStep5]
+CommandletClass=InternationalizationExport
+bExportLoc=true
\ No newline at end of file
