Added editor tutorial localization files and hooked them up to load with the editor.

diff --git a/Engine/Config/BaseEditor.ini b/Engine/Config/BaseEditor.ini
--- a/Engine/Config/BaseEditor.ini
+++ b/Engine/Config/BaseEditor.ini
@@ -1,5 +1,6 @@
 [Internationalization]
 +LocalizationPaths=../../../Engine/Content/Localization/Editor
++LocalizationPaths=../../../Engine/Content/Localization/EditorTutorials
 +ToolTipLocalizationPaths=../../../Engine/Content/Localization/Editor/ToolTips
 +PropertyNameLocalizationPaths=../../../Engine/Content/Localization/Editor/PropertyNames
 ShouldLoadLocalizedPropertyNames=True
diff --git a/Engine/Config/Localization/EditorTutorials.ini b/Engine/Config/Localization/EditorTutorials.ini
new file mode 100644
--- /dev/null
+++ b/Engine/Config/Localization/EditorTutorials.ini
@@ -0,0 +1,47 @@
+[CommonSettings]
+SourcePath=./Content/Localization/EditorTutorials
+DestinationPath=./Content/Localization/EditorTutorials
+ManifestName=EditorTutorials.manifest
+ArchiveName=EditorTutorials.archive
+PortableObjectName=EditorTutorials.po
+SourceCulture=en
+CulturesToGenerate=es
+CulturesToGenerate=en
+CulturesToGenerate=ja
+CulturesToGenerate=ko
+CulturesToGenerate=pt-BR
+CulturesToGenerate=zh-Hans-CN
+ModulesToPreload=IntroTutorials
+
+;Gather text from assets
+[GatherTextStep0]
+CommandletClass=GatherTextFromAssets
+IncludePaths=./Content/Tutorial/*
+PackageExtensions=*.umap
+PackageExtensions=*.upk
+PackageExtensions=*.uasset
+bFixBroken=false
+
+;Write Manifest
+[GatherTextStep1]
+CommandletClass=GenerateGatherManifest
+
+;Write Archives
+[GatherTextStep2]
+CommandletClass=GenerateGatherArchive
+bPurgeOldEmptyEntries=true
+
+;Import localized PO files
+[GatherTextStep3]
+CommandletClass=InternationalizationExport
+bImportLoc=true
+
+;Write Localized Text Resource
+[GatherTextStep4]
+CommandletClass=GenerateTextLocalizationResource
+ResourceName=EditorTutorials.locres
+
+;Export to PO files
+[GatherTextStep5]
+CommandletClass=InternationalizationExport
+bExportLoc=true
\ No newline at end of file
diff --git a/Engine/Config/Localization/Engine.ini b/Engine/Config/Localization/Engine.ini
--- a/Engine/Config/Localization/Engine.ini
+++ b/Engine/Config/Localization/Engine.ini
@@ -41,6 +41,7 @@ ExcludePaths=*/Content/TestPackages/*
 ExcludePaths=*/Content/QA_Assets/*
 ExcludePaths=*/Content/EngineSounds/*
 ExcludePaths=*/Content/Maps/Automation/*
+ExcludePaths=*/Content/Tutorial/*
 PackageExtensions=*.umap
 PackageExtensions=*.upk
 PackageExtensions=*.uasset
