L10N: Configuration file to allow the engine to load the editor tutorial loc directly from the raw archive and manifest files.

diff --git a/Engine/Config/Localization/RegenerateEditorTutorials.ini b/Engine/Config/Localization/RegenerateEditorTutorials.ini
new file mode 100644
--- /dev/null
+++ b/Engine/Config/Localization/RegenerateEditorTutorials.ini
@@ -0,0 +1,5 @@
+[RegenerateResources]
+SourcePath=./Content/Localization/EditorTutorials
+DestinationPath=./Content/Localization/EditorTutorials
+ManifestName=EditorTutorials.manifest
+ResourceName=EditorTutorials.locres
\ No newline at end of file
