Added Editor Tutorials target to various localization data generation configuration files.

diff --git a/Engine/Config/Localization/PortableObjectExport.ini b/Engine/Config/Localization/PortableObjectExport.ini
--- a/Engine/Config/Localization/PortableObjectExport.ini
+++ b/Engine/Config/Localization/PortableObjectExport.ini
@@ -37,4 +37,12 @@ SourcePath=./Content/Localization/PropertyNames
 DestinationPath=./Content/Localization/PropertyNames
 ManifestName=PropertyNames.manifest
 ArchiveName=PropertyNames.archive
-PortableObjectName=PropertyNames.po
\ No newline at end of file
+PortableObjectName=PropertyNames.po
+
+[GatherTextStep4]
+CommandletClass=InternationalizationExport
+SourcePath=./Content/Localization/EditorTutorials
+DestinationPath=./Content/Localization/EditorTutorials
+ManifestName=EditorTutorials.manifest
+ArchiveName=EditorTutorials.archive
+PortableObjectName=EditorTutorials.po
\ No newline at end of file
diff --git a/Engine/Config/Localization/PortableObjectImport.ini b/Engine/Config/Localization/PortableObjectImport.ini
--- a/Engine/Config/Localization/PortableObjectImport.ini
+++ b/Engine/Config/Localization/PortableObjectImport.ini
@@ -37,4 +37,12 @@ SourcePath=./Content/Localization/PropertyNames
 DestinationPath=./Content/Localization/PropertyNames
 ManifestName=PropertyNames.manifest
 ArchiveName=PropertyNames.archive
-PortableObjectName=PropertyNames.po
\ No newline at end of file
+PortableObjectName=PropertyNames.po
+
+[GatherTextStep4]
+CommandletClass=InternationalizationExport
+SourcePath=./Content/Localization/EditorTutorials
+DestinationPath=./Content/Localization/EditorTutorials
+ManifestName=EditorTutorials.manifest
+ArchiveName=EditorTutorials.archive
+PortableObjectName=EditorTutorials.po
\ No newline at end of file
diff --git a/Engine/Config/Localization/RepairData.ini b/Engine/Config/Localization/RepairData.ini
--- a/Engine/Config/Localization/RepairData.ini
+++ b/Engine/Config/Localization/RepairData.ini
@@ -64,4 +64,19 @@ CommandletClass=GenerateTextLocalizationResource
 SourcePath=./Content/Localization/PropertyNames
 DestinationPath=./Content/Localization/PropertyNames
 ManifestName=PropertyNames.manifest
-ResourceName=PropertyNames.locres
\ No newline at end of file
+ResourceName=PropertyNames.locres
+
+;Repair localization data for Editor Tutorials
+[GatherTextStep8]
+CommandletClass=RepairLocalizationData
+DestinationPath=./Content/Localization/PropertyNames
+ManifestName=PropertyNames.manifest
+ArchiveName=PropertyNames.archive
+
+;Write Localized Text Resource for Editor Tutorials
+[GatherTextStep9]
+CommandletClass=GenerateTextLocalizationResource
+SourcePath=./Content/Localization/EditorTutorials
+DestinationPath=./Content/Localization/EditorTutorials
+ManifestName=EditorTutorials.manifest
+ResourceName=EditorTutorials.locres
\ No newline at end of file
diff --git a/Engine/Config/Localization/ResourceFileGen.ini b/Engine/Config/Localization/ResourceFileGen.ini
--- a/Engine/Config/Localization/ResourceFileGen.ini
+++ b/Engine/Config/Localization/ResourceFileGen.ini
@@ -36,4 +36,12 @@ CommandletClass=GenerateTextLocalizationResource
 SourcePath=./Content/Localization/PropertyNames
 DestinationPath=./Content/Localization/PropertyNames
 ManifestName=PropertyNames.manifest
-ResourceName=PropertyNames.locres
\ No newline at end of file
+ResourceName=PropertyNames.locres
+
+;Write Localized Text Resource for Editor Tutorials
+[GatherTextStep4]
+CommandletClass=GenerateTextLocalizationResource
+SourcePath=./Content/Localization/EditorTutorials
+DestinationPath=./Content/Localization/EditorTutorials
+ManifestName=EditorTutorials.manifest
+ResourceName=EditorTutorials.locres
\ No newline at end of file
diff --git a/Engine/Config/Localization/WordCount.ini b/Engine/Config/Localization/WordCount.ini
--- a/Engine/Config/Localization/WordCount.ini
+++ b/Engine/Config/Localization/WordCount.ini
@@ -33,4 +33,11 @@ WordCountReportName=ToolTips.csv
 CommandletClass=GenerateTextLocalizationReport
 SourcePath=./Content/Localization/PropertyNames
 ManifestName=PropertyNames.manifest
-WordCountReportName=PropertyNames.csv
\ No newline at end of file
+WordCountReportName=PropertyNames.csv
+
+;Write Text Localization Report for Editor Tutorials
+[GatherTextStep4]
+CommandletClass=GenerateTextLocalizationReport
+SourcePath=./Content/Localization/EditorTutorials
+ManifestName=EditorTutorials.manifest
+WordCountReportName=EditorTutorials.csv
\ No newline at end of file
