Renamed Engine localization data for Simplified Chinese from zh-Hans-Cn to zh-CN due to convention.

diff --git a/Engine/Config/Localization/Editor.ini b/Engine/Config/Localization/Editor.ini
--- a/Engine/Config/Localization/Editor.ini
+++ b/Engine/Config/Localization/Editor.ini
@@ -11,7 +11,7 @@ CulturesToGenerate=es
 CulturesToGenerate=ja
 CulturesToGenerate=ko
 CulturesToGenerate=pt-BR
-CulturesToGenerate=zh-Hans-CN
+CulturesToGenerate=zh-CN
 
 ;Gather text from source code
 [GatherTextStep0]
diff --git a/Engine/Config/Localization/EditorTutorials.ini b/Engine/Config/Localization/EditorTutorials.ini
--- a/Engine/Config/Localization/EditorTutorials.ini
+++ b/Engine/Config/Localization/EditorTutorials.ini
@@ -10,7 +10,7 @@ CulturesToGenerate=en
 CulturesToGenerate=ja
 CulturesToGenerate=ko
 CulturesToGenerate=pt-BR
-CulturesToGenerate=zh-Hans-CN
+CulturesToGenerate=zh-CN
 ModulesToPreload=IntroTutorials
 
 ;Gather text from assets
diff --git a/Engine/Config/Localization/Engine.ini b/Engine/Config/Localization/Engine.ini
--- a/Engine/Config/Localization/Engine.ini
+++ b/Engine/Config/Localization/Engine.ini
@@ -10,7 +10,7 @@ CulturesToGenerate=es
 CulturesToGenerate=ja
 CulturesToGenerate=ko
 CulturesToGenerate=pt-BR
-CulturesToGenerate=zh-Hans-CN
+CulturesToGenerate=zh-CN
 
 ;Gather text from source code
 [GatherTextStep0]
diff --git a/Engine/Config/Localization/PortableObjectExport.ini b/Engine/Config/Localization/PortableObjectExport.ini
--- a/Engine/Config/Localization/PortableObjectExport.ini
+++ b/Engine/Config/Localization/PortableObjectExport.ini
@@ -5,7 +5,7 @@ CulturesToGenerate=es
 CulturesToGenerate=ja
 CulturesToGenerate=ko
 CulturesToGenerate=pt-BR
-CulturesToGenerate=zh-Hans-CN
+CulturesToGenerate=zh-CN
 
 [GatherTextStep0]
 CommandletClass=InternationalizationExport
diff --git a/Engine/Config/Localization/PortableObjectImport.ini b/Engine/Config/Localization/PortableObjectImport.ini
--- a/Engine/Config/Localization/PortableObjectImport.ini
+++ b/Engine/Config/Localization/PortableObjectImport.ini
@@ -5,7 +5,7 @@ CulturesToGenerate=es
 CulturesToGenerate=ja
 CulturesToGenerate=ko
 CulturesToGenerate=pt-BR
-CulturesToGenerate=zh-Hans-CN
+CulturesToGenerate=zh-CN
 
 [GatherTextStep0]
 CommandletClass=InternationalizationExport
diff --git a/Engine/Config/Localization/PropertyNames.ini b/Engine/Config/Localization/PropertyNames.ini
--- a/Engine/Config/Localization/PropertyNames.ini
+++ b/Engine/Config/Localization/PropertyNames.ini
@@ -10,7 +10,7 @@ CulturesToGenerate=en
 CulturesToGenerate=ja
 CulturesToGenerate=ko
 CulturesToGenerate=pt-BR
-CulturesToGenerate=zh-Hans-CN
+CulturesToGenerate=zh-CN
 
 ;Gather text from metadata
 [GatherTextStep0]
diff --git a/Engine/Config/Localization/RepairData.ini b/Engine/Config/Localization/RepairData.ini
--- a/Engine/Config/Localization/RepairData.ini
+++ b/Engine/Config/Localization/RepairData.ini
@@ -4,7 +4,7 @@ CulturesToGenerate=es
 CulturesToGenerate=ja
 CulturesToGenerate=ko
 CulturesToGenerate=pt-BR
-CulturesToGenerate=zh-Hans-CN
+CulturesToGenerate=zh-CN
 
 ;Repair localization data for Engine
 [GatherTextStep0]
diff --git a/Engine/Config/Localization/ResourceFileGen.ini b/Engine/Config/Localization/ResourceFileGen.ini
--- a/Engine/Config/Localization/ResourceFileGen.ini
+++ b/Engine/Config/Localization/ResourceFileGen.ini
@@ -4,7 +4,7 @@ CulturesToGenerate=es
 CulturesToGenerate=ja
 CulturesToGenerate=ko
 CulturesToGenerate=pt-BR
-CulturesToGenerate=zh-Hans-CN
+CulturesToGenerate=zh-CN
 
 ;Write Localized Text Resource for Engine
 [GatherTextStep0]
diff --git a/Engine/Config/Localization/ToolTips.ini b/Engine/Config/Localization/ToolTips.ini
--- a/Engine/Config/Localization/ToolTips.ini
+++ b/Engine/Config/Localization/ToolTips.ini
@@ -10,7 +10,7 @@ CulturesToGenerate=en
 CulturesToGenerate=ja
 CulturesToGenerate=ko
 CulturesToGenerate=pt-BR
-CulturesToGenerate=zh-Hans-CN
+CulturesToGenerate=zh-CN
 
 ;Gather text from metadata
 [GatherTextStep0]
diff --git a/Engine/Config/Localization/WordCount.ini b/Engine/Config/Localization/WordCount.ini
--- a/Engine/Config/Localization/WordCount.ini
+++ b/Engine/Config/Localization/WordCount.ini
@@ -5,7 +5,7 @@ CulturesToGenerate=es
 CulturesToGenerate=ja
 CulturesToGenerate=ko
 CulturesToGenerate=pt-BR
-CulturesToGenerate=zh-Hans-CN
+CulturesToGenerate=zh-CN
 
 ;Write Text Localization Report for Engine
 [GatherTextStep0]
