Temporarily disabled word count and conflict report generation from individual legacy localization configuration scripts for engine targets until promotion occurs. Also corrected a few warnings about deprecated settings in said localization configuration scripts.

diff --git a/Engine/Config/Localization/Editor.ini b/Engine/Config/Localization/Editor.ini
--- a/Engine/Config/Localization/Editor.ini
+++ b/Engine/Config/Localization/Editor.ini
@@ -1,4 +1,4 @@
-[CommonSettings]
+﻿[CommonSettings]
 ManifestDependencies=Content/Localization/Engine/Engine.manifest
 SourcePath=Content/Localization/Editor
 DestinationPath=Content/Localization/Editor
@@ -26,13 +26,13 @@ FileNameFilters=*.mm
 ;Gather text from metadata
 [GatherTextStep1]
 CommandletClass=GatherTextFromMetaData
-IncludePaths=Source/Editor/*
-IncludePaths=Source/Runtime/*
-IncludePaths=Source/Developer/*
-ExcludePaths=Config/Localization/*
-ExcludePaths=Source/Developer/NoRedist/CommunityPortalServices/*
-ExcludePaths=Source/Developer/NoRedist/UnrealEngineLauncherServices/*
-ExcludePaths=Source/Developer/NoRedist/BuildPatchServices/*
+IncludePathFilters=Source/Editor/*
+IncludePathFilters=Source/Runtime/*
+IncludePathFilters=Source/Developer/*
+ExcludePathFilters=Config/Localization/*
+ExcludePathFilters=Source/Developer/NoRedist/CommunityPortalServices/*
+ExcludePathFilters=Source/Developer/NoRedist/UnrealEngineLauncherServices/*
+ExcludePathFilters=Source/Developer/NoRedist/BuildPatchServices/*
 InputKeys=Category
 OutputNamespaces=UObjectCategories
 OutputKeys="{MetaDataValue}"
@@ -40,8 +40,8 @@ OutputKeys="{MetaDataValue}"
 ;Gather text from assets
 [GatherTextStep2]
 CommandletClass=GatherTextFromAssets
-IncludePaths=Content/Editor/*
-IncludePaths=Content/Editor*
+IncludePathFilters=Content/Editor/*
+IncludePathFilters=Content/Editor*
 PackageFileNameFilters=*.umap
 PackageFileNameFilters=*.upk
 PackageFileNameFilters=*.uasset
@@ -72,11 +72,11 @@ ResourceName=Editor.locres
 CommandletClass=InternationalizationExport
 bExportLoc=true
 
-;Write Text Localization Report
-[GatherTextStep8]
-CommandletClass=GenerateTextLocalizationReport
-DestinationPath=\\epicgames.net\root\UE3\Localization\WordCounts
-bWordCountReport=true
-WordCountReportName=Editor.csv
-bConflictReport=true
-ConflictReportName=EditorConflicts.txt
\ No newline at end of file
+;;Write Text Localization Report
+;[GatherTextStep8]
+;CommandletClass=GenerateTextLocalizationReport
+;DestinationPath=\\epicgames.net\root\UE3\Localization\WordCounts
+;bWordCountReport=true
+;WordCountReportName=Editor.csv
+;bConflictReport=true
+;ConflictReportName=EditorConflicts.txt
\ No newline at end of file
diff --git a/Engine/Config/Localization/EditorTutorials.ini b/Engine/Config/Localization/EditorTutorials.ini
--- a/Engine/Config/Localization/EditorTutorials.ini
+++ b/Engine/Config/Localization/EditorTutorials.ini
@@ -1,4 +1,4 @@
-[CommonSettings]
+﻿[CommonSettings]
 SourcePath=Content/Localization/EditorTutorials
 DestinationPath=Content/Localization/EditorTutorials
 ManifestName=EditorTutorials.manifest
@@ -16,7 +16,7 @@ ModulesToPreload=IntroTutorials
 ;Gather text from assets
 [GatherTextStep0]
 CommandletClass=GatherTextFromAssets
-IncludePaths=Content/Tutorial/*
+IncludePathFilters=Content/Tutorial/*
 PackageFileNameFilters=*.umap
 PackageFileNameFilters=*.upk
 PackageFileNameFilters=*.uasset
@@ -46,11 +46,11 @@ ResourceName=EditorTutorials.locres
 CommandletClass=InternationalizationExport
 bExportLoc=true
 
-;Write Text Localization Report
-[GatherTextStep6]
-CommandletClass=GenerateTextLocalizationReport
-DestinationPath=\\epicgames.net\root\UE3\Localization\WordCounts
-bWordCountReport=true
-WordCountReportName=EditorTutorials.csv
-bConflictReport=true
-ConflictReportName=EditorTutorialsConflicts.txt
\ No newline at end of file
+;;Write Text Localization Report
+;[GatherTextStep6]
+;CommandletClass=GenerateTextLocalizationReport
+;DestinationPath=\\epicgames.net\root\UE3\Localization\WordCounts
+;bWordCountReport=true
+;WordCountReportName=EditorTutorials.csv
+;bConflictReport=true
+;ConflictReportName=EditorTutorialsConflicts.txt
\ No newline at end of file
diff --git a/Engine/Config/Localization/Engine.ini b/Engine/Config/Localization/Engine.ini
--- a/Engine/Config/Localization/Engine.ini
+++ b/Engine/Config/Localization/Engine.ini
@@ -1,4 +1,4 @@
-[CommonSettings]
+﻿[CommonSettings]
 SourcePath=Content/Localization/Engine
 DestinationPath=Content/Localization/Engine
 ManifestName=Engine.manifest
@@ -19,10 +19,10 @@ SearchDirectoryPaths=Source/Runtime/
 SearchDirectoryPaths=Source/Rocket/
 SearchDirectoryPaths=Source/Developer/
 SearchDirectoryPaths=Config/
-ExcludePaths=Config/Localization/*
-ExcludePaths=Source/Developer/NoRedist/CommunityPortalServices/*
-ExcludePaths=Source/Developer/NoRedist/UnrealEngineLauncherServices/*
-ExcludePaths=Source/Developer/NoRedist/BuildPatchServices/*
+ExcludePathFilters=Config/Localization/*
+ExcludePathFilters=Source/Developer/NoRedist/CommunityPortalServices/*
+ExcludePathFilters=Source/Developer/NoRedist/UnrealEngineLauncherServices/*
+ExcludePathFilters=Source/Developer/NoRedist/BuildPatchServices/*
 FileNameFilters=*.cpp
 FileNameFilters=*.h
 FileNameFilters=*.c
@@ -33,15 +33,15 @@ FileNameFilters=*.ini
 ;Gather text from assets
 [GatherTextStep1]
 CommandletClass=GatherTextFromAssets
-IncludePaths=Content/*
-ExcludePaths=Content/Editor/*
-ExcludePaths=Content/Developers/*
-ExcludePaths=Content/Localization/*
-ExcludePaths=Content/TestPackages/*
-ExcludePaths=Content/QA_Assets/*
-ExcludePaths=Content/EngineSounds/*
-ExcludePaths=Content/Maps/Automation/*
-ExcludePaths=Content/Tutorial/*
+IncludePathFilters=Content/*
+ExcludePathFilters=Content/Editor/*
+ExcludePathFilters=Content/Developers/*
+ExcludePathFilters=Content/Localization/*
+ExcludePathFilters=Content/TestPackages/*
+ExcludePathFilters=Content/QA_Assets/*
+ExcludePathFilters=Content/EngineSounds/*
+ExcludePathFilters=Content/Maps/Automation/*
+ExcludePathFilters=Content/Tutorial/*
 PackageFileNameFilters=*.umap
 PackageFileNameFilters=*.upk
 PackageFileNameFilters=*.uasset
@@ -72,11 +72,11 @@ ResourceName=Engine.locres
 CommandletClass=InternationalizationExport
 bExportLoc=true
 
-;Write Text Localization Report
-[GatherTextStep7]
-CommandletClass=GenerateTextLocalizationReport
-DestinationPath=\\partnersdc\NetworkShare\Users\SaulDavidAbreu
-bWordCountReport=true
-WordCountReportName=Engine.csv
-bConflictReport=true
-ConflictReportName=EngineConflicts.txt
\ No newline at end of file
+;;Write Text Localization Report
+;[GatherTextStep7]
+;CommandletClass=GenerateTextLocalizationReport
+;DestinationPath=\\partnersdc\NetworkShare\Users\SaulDavidAbreu
+;bWordCountReport=true
+;WordCountReportName=Engine.csv
+;bConflictReport=true
+;ConflictReportName=EngineConflicts.txt
\ No newline at end of file
diff --git a/Engine/Config/Localization/Keywords.ini b/Engine/Config/Localization/Keywords.ini
--- a/Engine/Config/Localization/Keywords.ini
+++ b/Engine/Config/Localization/Keywords.ini
@@ -1,4 +1,4 @@
-[CommonSettings]
+﻿[CommonSettings]
 SourcePath=./Content/Localization/Keywords
 DestinationPath=./Content/Localization/Keywords
 ManifestName=Keywords.manifest
@@ -15,11 +15,11 @@ CulturesToGenerate=zh-CN
 ;Gather text from metadata
 [GatherTextStep0]
 CommandletClass=GatherTextFromMetaData
-IncludePaths=Source/Editor/*
-IncludePaths=Source/Runtime/*
-IncludePaths=Source/Developer/*
-ExcludePaths=Config/Localization/*
-ExcludePaths=Source/Developer/NoRedist/CommunityPortalServices/*
+IncludePathFilters=Source/Editor/*
+IncludePathFilters=Source/Runtime/*
+IncludePathFilters=Source/Developer/*
+ExcludePathFilters=Config/Localization/*
+ExcludePathFilters=Source/Developer/NoRedist/CommunityPortalServices/*
 InputKeys=Keywords
 OutputNamespaces=UObjectKeywords
 OutputKeys="{FieldPath}"
@@ -48,11 +48,11 @@ ResourceName=Keywords.locres
 CommandletClass=InternationalizationExport
 bExportLoc=true
 
-;Write Text Localization Report
-[GatherTextStep6]
-CommandletClass=GenerateTextLocalizationReport
-DestinationPath=\\epicgames.net\root\UE3\Localization\WordCounts
-bWordCountReport=true
-WordCountReportName=Keywords.csv
-bConflictReport=true
-ConflictReportName=KeywordsConflicts.txt
\ No newline at end of file
+;;Write Text Localization Report
+;[GatherTextStep6]
+;CommandletClass=GenerateTextLocalizationReport
+;DestinationPath=\\epicgames.net\root\UE3\Localization\WordCounts
+;bWordCountReport=true
+;WordCountReportName=Keywords.csv
+;bConflictReport=true
+;ConflictReportName=KeywordsConflicts.txt
\ No newline at end of file
diff --git a/Engine/Config/Localization/PropertyNames.ini b/Engine/Config/Localization/PropertyNames.ini
--- a/Engine/Config/Localization/PropertyNames.ini
+++ b/Engine/Config/Localization/PropertyNames.ini
@@ -1,4 +1,4 @@
-[CommonSettings]
+﻿[CommonSettings]
 SourcePath=Content/Localization/PropertyNames
 DestinationPath=Content/Localization/PropertyNames
 ManifestName=PropertyNames.manifest
@@ -15,11 +15,11 @@ CulturesToGenerate=zh-CN
 ;Gather text from metadata
 [GatherTextStep0]
 CommandletClass=GatherTextFromMetaData
-IncludePaths=Source/Editor/*
-IncludePaths=Source/Runtime/*
-IncludePaths=Source/Developer/*
-ExcludePaths=Config/Localization/*
-ExcludePaths=Source/Developer/NoRedist/CommunityPortalServices/*
+IncludePathFilters=Source/Editor/*
+IncludePathFilters=Source/Runtime/*
+IncludePathFilters=Source/Developer/*
+ExcludePathFilters=Config/Localization/*
+ExcludePathFilters=Source/Developer/NoRedist/CommunityPortalServices/*
 InputKeys=DisplayName
 OutputNamespaces=UObjectDisplayNames
 OutputKeys="{FieldPath}"
@@ -48,11 +48,11 @@ ResourceName=PropertyNames.locres
 CommandletClass=InternationalizationExport
 bExportLoc=true
 
-;Write Text Localization Report
-[GatherTextStep6]
-CommandletClass=GenerateTextLocalizationReport
-DestinationPath=\\epicgames.net\root\UE3\Localization\WordCounts
-bWordCountReport=true
-WordCountReportName=PropertyNames.csv
-bConflictReport=true
-ConflictReportName=PropertyNamesConflicts.txt
\ No newline at end of file
+;;Write Text Localization Report
+;[GatherTextStep6]
+;CommandletClass=GenerateTextLocalizationReport
+;DestinationPath=\\epicgames.net\root\UE3\Localization\WordCounts
+;bWordCountReport=true
+;WordCountReportName=PropertyNames.csv
+;bConflictReport=true
+;ConflictReportName=PropertyNamesConflicts.txt
\ No newline at end of file
diff --git a/Engine/Config/Localization/ToolTips.ini b/Engine/Config/Localization/ToolTips.ini
--- a/Engine/Config/Localization/ToolTips.ini
+++ b/Engine/Config/Localization/ToolTips.ini
@@ -1,4 +1,4 @@
-[CommonSettings]
+﻿[CommonSettings]
 SourcePath=Content/Localization/ToolTips
 DestinationPath=Content/Localization/ToolTips
 ManifestName=ToolTips.manifest
@@ -15,11 +15,11 @@ CulturesToGenerate=zh-CN
 ;Gather text from metadata
 [GatherTextStep0]
 CommandletClass=GatherTextFromMetaData
-IncludePaths=Source/Editor/*
-IncludePaths=Source/Runtime/*
-IncludePaths=Source/Developer/*
-ExcludePaths=Config/Localization/*
-ExcludePaths=Source/Developer/NoRedist/CommunityPortalServices/*
+IncludePathFilters=Source/Editor/*
+IncludePathFilters=Source/Runtime/*
+IncludePathFilters=Source/Developer/*
+ExcludePathFilters=Config/Localization/*
+ExcludePathFilters=Source/Developer/NoRedist/CommunityPortalServices/*
 InputKeys=ToolTip
 OutputNamespaces=UObjectToolTips
 OutputKeys="{FieldPath}"
@@ -51,11 +51,11 @@ ResourceName=ToolTips.locres
 CommandletClass=InternationalizationExport
 bExportLoc=true
 
-;Write Text Localization Report
-[GatherTextStep6]
-CommandletClass=GenerateTextLocalizationReport
-DestinationPath=\\epicgames.net\root\UE3\Localization\WordCounts
-bWordCountReport=true
-WordCountReportName=ToolTips.csv
-bConflictReport=true
-ConflictReportName=ToolTipsConflicts.txt
\ No newline at end of file
+;;Write Text Localization Report
+;[GatherTextStep6]
+;CommandletClass=GenerateTextLocalizationReport
+;DestinationPath=\\epicgames.net\root\UE3\Localization\WordCounts
+;bWordCountReport=true
+;WordCountReportName=ToolTips.csv
+;bConflictReport=true
+;ConflictReportName=ToolTipsConflicts.txt
\ No newline at end of file
