Temporary change of settings in Localization/Engine.ini of key ExcludePathFilters to ExcludePaths until automation catches up with the relevant fix.

diff --git a/Engine/Config/Localization/Engine.ini b/Engine/Config/Localization/Engine.ini
--- a/Engine/Config/Localization/Engine.ini
+++ b/Engine/Config/Localization/Engine.ini
@@ -19,10 +19,10 @@ SearchDirectoryPaths=Source/Runtime/
 SearchDirectoryPaths=Source/Rocket/
 SearchDirectoryPaths=Source/Developer/
 SearchDirectoryPaths=Config/
-ExcludePathFilters=Config/Localization/*
-ExcludePathFilters=Source/Developer/NoRedist/CommunityPortalServices/*
-ExcludePathFilters=Source/Developer/NoRedist/UnrealEngineLauncherServices/*
-ExcludePathFilters=Source/Developer/NoRedist/BuildPatchServices/*
+ExcludePaths=Config/Localization/*
+ExcludePaths=Source/Developer/NoRedist/CommunityPortalServices/*
+ExcludePaths=Source/Developer/NoRedist/UnrealEngineLauncherServices/*
+ExcludePaths=Source/Developer/NoRedist/BuildPatchServices/*
 FileNameFilters=*.cpp
 FileNameFilters=*.h
 FileNameFilters=*.c
@@ -45,7 +45,7 @@ ExcludePathFilters=Content/Tutorial/*
 PackageFileNameFilters=*.umap
 PackageFileNameFilters=*.upk
 PackageFileNameFilters=*.uasset
-bFixBroken=false
+bFixBroken=true
 ShouldGatherBlueprintPinNames=true
 
 ;Write Manifest
