Removed old method of gathering Category metadata in favor of a new method.

Deleted broken code for finding localized version of details category display name.

diff --git a/Engine/Config/Localization/Editor.ini b/Engine/Config/Localization/Editor.ini
--- a/Engine/Config/Localization/Editor.ini
+++ b/Engine/Config/Localization/Editor.ini
@@ -23,22 +23,8 @@ FileNameFilters=*.c
 FileNameFilters=*.inl
 FileNameFilters=*.mm
 
-;Gather text from metadata
-[GatherTextStep1]
-CommandletClass=GatherTextFromMetaData
-IncludePathFilters=Source/Editor/*
-IncludePathFilters=Source/Runtime/*
-IncludePathFilters=Source/Developer/*
-ExcludePathFilters=Config/Localization/*
-ExcludePathFilters=Source/Developer/NoRedist/CommunityPortalServices/*
-ExcludePathFilters=Source/Developer/NoRedist/UnrealEngineLauncherServices/*
-ExcludePathFilters=Source/Developer/NoRedist/BuildPatchServices/*
-InputKeys=Category
-OutputNamespaces=UObjectCategories
-OutputKeys="{MetaDataValue}"
-
 ;Gather text from assets
-[GatherTextStep2]
+[GatherTextStep1]
 CommandletClass=GatherTextFromAssets
 IncludePathFilters=Content/Editor/*
 IncludePathFilters=Content/Editor*
@@ -49,31 +35,31 @@ bFixBroken=false
 ShouldGatherBlueprintPinNames=true
 
 ;Write Manifest
-[GatherTextStep3]
+[GatherTextStep2]
 CommandletClass=GenerateGatherManifest
 
 ;Write Archives
-[GatherTextStep4]
+[GatherTextStep3]
 CommandletClass=GenerateGatherArchive
 bPurgeOldEmptyEntries=true
 
 ;Import localized PO files
-[GatherTextStep5]
+[GatherTextStep4]
 CommandletClass=InternationalizationExport
 bImportLoc=true
 
 ;Write Localized Text Resource
-[GatherTextStep6]
+[GatherTextStep5]
 CommandletClass=GenerateTextLocalizationResource
 ResourceName=Editor.locres
 
 ;Export to PO files
-[GatherTextStep7]
+[GatherTextStep6]
 CommandletClass=InternationalizationExport
 bExportLoc=true
 
 ;Write Text Localization Report
-[GatherTextStep8]
+[GatherTextStep7]
 CommandletClass=GenerateTextLocalizationReport
 DestinationPath=\\epicgames.net\root\UE3\Localization\WordCounts
 bWordCountReport=true
diff --git a/Engine/Source/Editor/PropertyEditor/Private/CategoryPropertyNode.cpp b/Engine/Source/Editor/PropertyEditor/Private/CategoryPropertyNode.cpp
--- a/Engine/Source/Editor/PropertyEditor/Private/CategoryPropertyNode.cpp
+++ b/Engine/Source/Editor/PropertyEditor/Private/CategoryPropertyNode.cpp
@@ -26,18 +26,6 @@ bool FCategoryPropertyNode::IsSubcategory() const
 FText FCategoryPropertyNode::GetDisplayName() const 
 {
 	FString SubcategoryName = GetSubcategoryName();
-	if ( FPropertySettings::Get().ShowFriendlyPropertyNames() )
-	{
-		//if there is a localized version, return that
-		FText LocalizedName;
-		if ( FText::FindText( TEXT("UObjectCategories"), *SubcategoryName, /*OUT*/LocalizedName ) )
-		{
-			return LocalizedName;
-		}
-
-		//if in "readable display name mode" return that
-		return FText::FromString( FName::NameToDisplayString( SubcategoryName, false ) );
-	}
 	return FText::FromString( SubcategoryName );
 }	
 /**
