Fixed deprecated SourceCulture settings in localizaton configuration scripts for Category and PropertyNames localization targets of UE4. Replaced with the appropriate NativeCulture setting instead.

diff --git a/Engine/Config/Localization/Category.ini b/Engine/Config/Localization/Category.ini
--- a/Engine/Config/Localization/Category.ini
+++ b/Engine/Config/Localization/Category.ini
@@ -4,7 +4,7 @@ DestinationPath=./Content/Localization/Category
 ManifestName=Category.manifest
 ArchiveName=Category.archive
 PortableObjectName=Category.po
-SourceCulture=en
+NativeCulture=en
 CulturesToGenerate=es
 CulturesToGenerate=en
 CulturesToGenerate=ja
diff --git a/Engine/Config/Localization/PropertyNames.ini b/Engine/Config/Localization/PropertyNames.ini
--- a/Engine/Config/Localization/PropertyNames.ini
+++ b/Engine/Config/Localization/PropertyNames.ini
@@ -4,7 +4,7 @@ DestinationPath=Content/Localization/PropertyNames
 ManifestName=PropertyNames.manifest
 ArchiveName=PropertyNames.archive
 PortableObjectName=PropertyNames.po
-SourceCulture=en
+NativeCulture=en
 CulturesToGenerate=es
 CulturesToGenerate=en
 CulturesToGenerate=ja
