Updated engine localization configuration scripts to use the new setting to gather text from editor-only data.

diff --git a/Engine/Config/Localization/Category.ini b/Engine/Config/Localization/Category.ini
--- a/Engine/Config/Localization/Category.ini
+++ b/Engine/Config/Localization/Category.ini
@@ -23,6 +23,7 @@ ExcludePathFilters=Source/Developer/NoRedist/CommunityPortalServices/*
 InputKeys=Category
 OutputNamespaces=UObjectCategory
 OutputKeys="{FieldPath}"
+ShouldGatherFromEditorOnlyData=true
 
 ;Write Manifest
 [GatherTextStep1]
diff --git a/Engine/Config/Localization/Editor.ini b/Engine/Config/Localization/Editor.ini
--- a/Engine/Config/Localization/Editor.ini
+++ b/Engine/Config/Localization/Editor.ini
@@ -32,7 +32,7 @@ PackageFileNameFilters=*.umap
 PackageFileNameFilters=*.upk
 PackageFileNameFilters=*.uasset
 bFixBroken=false
-ShouldGatherBlueprintPinNames=true
+ShouldGatherFromEditorOnlyData=true
 
 ;Write Manifest
 [GatherTextStep2]
diff --git a/Engine/Config/Localization/EditorTutorials.ini b/Engine/Config/Localization/EditorTutorials.ini
--- a/Engine/Config/Localization/EditorTutorials.ini
+++ b/Engine/Config/Localization/EditorTutorials.ini
@@ -21,6 +21,7 @@ PackageFileNameFilters=*.umap
 PackageFileNameFilters=*.upk
 PackageFileNameFilters=*.uasset
 bFixBroken=false
+ShouldGatherFromEditorOnlyData=true
 
 ;Write Manifest
 [GatherTextStep1]
diff --git a/Engine/Config/Localization/Engine.ini b/Engine/Config/Localization/Engine.ini
--- a/Engine/Config/Localization/Engine.ini
+++ b/Engine/Config/Localization/Engine.ini
@@ -46,7 +46,7 @@ PackageFileNameFilters=*.umap
 PackageFileNameFilters=*.upk
 PackageFileNameFilters=*.uasset
 bFixBroken=false
-ShouldGatherBlueprintPinNames=true
+ShouldGatherFromEditorOnlyData=true
 
 ;Write Manifest
 [GatherTextStep2]
diff --git a/Engine/Config/Localization/Keywords.ini b/Engine/Config/Localization/Keywords.ini
--- a/Engine/Config/Localization/Keywords.ini
+++ b/Engine/Config/Localization/Keywords.ini
@@ -23,6 +23,7 @@ ExcludePathFilters=Source/Developer/NoRedist/CommunityPortalServices/*
 InputKeys=Keywords
 OutputNamespaces=UObjectKeywords
 OutputKeys="{FieldPath}"
+ShouldGatherFromEditorOnlyData=true
 
 ;Write Manifest
 [GatherTextStep1]
diff --git a/Engine/Config/Localization/PropertyNames.ini b/Engine/Config/Localization/PropertyNames.ini
--- a/Engine/Config/Localization/PropertyNames.ini
+++ b/Engine/Config/Localization/PropertyNames.ini
@@ -23,6 +23,7 @@ ExcludePathFilters=Source/Developer/NoRedist/CommunityPortalServices/*
 InputKeys=DisplayName
 OutputNamespaces=UObjectDisplayNames
 OutputKeys="{FieldPath}"
+ShouldGatherFromEditorOnlyData=true
 
 ;Write Manifest
 [GatherTextStep1]
diff --git a/Engine/Config/Localization/ToolTips.ini b/Engine/Config/Localization/ToolTips.ini
--- a/Engine/Config/Localization/ToolTips.ini
+++ b/Engine/Config/Localization/ToolTips.ini
@@ -26,6 +26,7 @@ OutputKeys="{FieldPath}"
 InputKeys=ShortTooltip
 OutputNamespaces=UObjectShortTooltips
 OutputKeys="{FieldPath}"
+ShouldGatherFromEditorOnlyData=true
 
 ;Write Manifest
 [GatherTextStep1]
