Included plugins in the editor/engine localization gather

diff --git a/Engine/Config/Localization/Category.ini b/Engine/Config/Localization/Category.ini
--- a/Engine/Config/Localization/Category.ini
+++ b/Engine/Config/Localization/Category.ini
@@ -19,8 +19,8 @@ CommandletClass=GatherTextFromMetaData
 IncludePathFilters=Source/Editor/*
 IncludePathFilters=Source/Runtime/*
 IncludePathFilters=Source/Developer/*
-ExcludePathFilters=Config/Localization/*
-ExcludePathFilters=Source/Developer/NoRedist/CommunityPortalServices/*
+ExcludePathFilters=*/NoRedist/*
+ExcludePathFilters=*/NotForLicensees/*
 InputKeys=Category
 OutputNamespaces=UObjectCategory
 OutputKeys="{FieldPath}"
diff --git a/Engine/Config/Localization/Editor.ini b/Engine/Config/Localization/Editor.ini
--- a/Engine/Config/Localization/Editor.ini
+++ b/Engine/Config/Localization/Editor.ini
@@ -18,6 +18,9 @@ CulturesToGenerate=zh-CN
 [GatherTextStep0]
 CommandletClass=GatherTextFromSource
 SearchDirectoryPaths=Source/Editor/
+SearchDirectoryPaths=Plugins/
+ExcludePathFilters=*/NoRedist/*
+ExcludePathFilters=*/NotForLicensees/*
 FileNameFilters=*.cpp
 FileNameFilters=*.h
 FileNameFilters=*.c
@@ -31,7 +34,6 @@ CommandletClass=GatherTextFromAssets
 IncludePathFilters=Content/Editor/*
 IncludePathFilters=Content/Editor*
 PackageFileNameFilters=*.umap
-PackageFileNameFilters=*.upk
 PackageFileNameFilters=*.uasset
 bFixBroken=false
 ShouldGatherFromEditorOnlyData=true
diff --git a/Engine/Config/Localization/Engine.ini b/Engine/Config/Localization/Engine.ini
--- a/Engine/Config/Localization/Engine.ini
+++ b/Engine/Config/Localization/Engine.ini
@@ -17,11 +17,16 @@ CulturesToGenerate=zh-CN
 [GatherTextStep0]
 CommandletClass=GatherTextFromSource
 SearchDirectoryPaths=Source/Runtime/
-SearchDirectoryPaths=Source/Rocket/
 SearchDirectoryPaths=Source/Developer/
+SearchDirectoryPaths=Plugins/
 SearchDirectoryPaths=Config/
 ExcludePathFilters=Config/Localization/*
 ExcludePathFilters=Source/Runtime/Online/BuildPatchServices/*
+ExcludePathFilters=Plugins/*Editor*/*
+ExcludePathFilters=Plugins/Online/*
+ExcludePathFilters=Plugins/Runtime/HTTPChunkInstaller/*
+ExcludePathFilters=*/NoRedist/*
+ExcludePathFilters=*/NotForLicensees/*
 FileNameFilters=*.cpp
 FileNameFilters=*.h
 FileNameFilters=*.c
@@ -43,7 +48,6 @@ ExcludePathFilters=Content/EngineSounds/*
 ExcludePathFilters=Content/Maps/Automation/*
 ExcludePathFilters=Content/Tutorial/*
 PackageFileNameFilters=*.umap
-PackageFileNameFilters=*.upk
 PackageFileNameFilters=*.uasset
 bFixBroken=false
 ShouldGatherFromEditorOnlyData=true
diff --git a/Engine/Config/Localization/Keywords.ini b/Engine/Config/Localization/Keywords.ini
--- a/Engine/Config/Localization/Keywords.ini
+++ b/Engine/Config/Localization/Keywords.ini
@@ -19,8 +19,8 @@ CommandletClass=GatherTextFromMetaData
 IncludePathFilters=Source/Editor/*
 IncludePathFilters=Source/Runtime/*
 IncludePathFilters=Source/Developer/*
-ExcludePathFilters=Config/Localization/*
-ExcludePathFilters=Source/Developer/NoRedist/CommunityPortalServices/*
+ExcludePathFilters=*/NoRedist/*
+ExcludePathFilters=*/NotForLicensees/*
 InputKeys=Keywords
 OutputNamespaces=UObjectKeywords
 OutputKeys="{FieldPath}"
diff --git a/Engine/Config/Localization/PropertyNames.ini b/Engine/Config/Localization/PropertyNames.ini
--- a/Engine/Config/Localization/PropertyNames.ini
+++ b/Engine/Config/Localization/PropertyNames.ini
@@ -19,8 +19,8 @@ CommandletClass=GatherTextFromMetaData
 IncludePathFilters=Source/Editor/*
 IncludePathFilters=Source/Runtime/*
 IncludePathFilters=Source/Developer/*
-ExcludePathFilters=Config/Localization/*
-ExcludePathFilters=Source/Developer/NoRedist/CommunityPortalServices/*
+ExcludePathFilters=*/NoRedist/*
+ExcludePathFilters=*/NotForLicensees/*
 InputKeys=DisplayName
 OutputNamespaces=UObjectDisplayNames
 OutputKeys="{FieldPath}"
diff --git a/Engine/Config/Localization/ToolTips.ini b/Engine/Config/Localization/ToolTips.ini
--- a/Engine/Config/Localization/ToolTips.ini
+++ b/Engine/Config/Localization/ToolTips.ini
@@ -20,8 +20,8 @@ CommandletClass=GatherTextFromMetaData
 IncludePathFilters=Source/Editor/*
 IncludePathFilters=Source/Runtime/*
 IncludePathFilters=Source/Developer/*
-ExcludePathFilters=Config/Localization/*
-ExcludePathFilters=Source/Developer/NoRedist/CommunityPortalServices/*
+ExcludePathFilters=*/NoRedist/*
+ExcludePathFilters=*/NotForLicensees/*
 InputKeys=ToolTip
 OutputNamespaces=UObjectToolTips
 OutputKeys="{FieldPath}"
