Included plugins in the editor/engine localization gather

diff --git a/Engine/Config/Localization/Category.ini b/Engine/Config/Localization/Category.ini
--- a/Engine/Config/Localization/Category.ini
+++ b/Engine/Config/Localization/Category.ini
@@ -19,6 +19,7 @@ CommandletClass=GatherTextFromMetaData
 IncludePathFilters=Source/Editor/*
 IncludePathFilters=Source/Runtime/*
 IncludePathFilters=Source/Developer/*
+IncludePathFilters=Plugins/*
 ExcludePathFilters=*/NoRedist/*
 ExcludePathFilters=*/NotForLicensees/*
 InputKeys=Category
diff --git a/Engine/Config/Localization/Keywords.ini b/Engine/Config/Localization/Keywords.ini
--- a/Engine/Config/Localization/Keywords.ini
+++ b/Engine/Config/Localization/Keywords.ini
@@ -19,6 +19,7 @@ CommandletClass=GatherTextFromMetaData
 IncludePathFilters=Source/Editor/*
 IncludePathFilters=Source/Runtime/*
 IncludePathFilters=Source/Developer/*
+IncludePathFilters=Plugins/*
 ExcludePathFilters=*/NoRedist/*
 ExcludePathFilters=*/NotForLicensees/*
 InputKeys=Keywords
diff --git a/Engine/Config/Localization/PropertyNames.ini b/Engine/Config/Localization/PropertyNames.ini
--- a/Engine/Config/Localization/PropertyNames.ini
+++ b/Engine/Config/Localization/PropertyNames.ini
@@ -19,6 +19,7 @@ CommandletClass=GatherTextFromMetaData
 IncludePathFilters=Source/Editor/*
 IncludePathFilters=Source/Runtime/*
 IncludePathFilters=Source/Developer/*
+IncludePathFilters=Plugins/*
 ExcludePathFilters=*/NoRedist/*
 ExcludePathFilters=*/NotForLicensees/*
 InputKeys=DisplayName
diff --git a/Engine/Config/Localization/ToolTips.ini b/Engine/Config/Localization/ToolTips.ini
--- a/Engine/Config/Localization/ToolTips.ini
+++ b/Engine/Config/Localization/ToolTips.ini
@@ -20,6 +20,7 @@ CommandletClass=GatherTextFromMetaData
 IncludePathFilters=Source/Editor/*
 IncludePathFilters=Source/Runtime/*
 IncludePathFilters=Source/Developer/*
+IncludePathFilters=Plugins/*
 ExcludePathFilters=*/NoRedist/*
 ExcludePathFilters=*/NotForLicensees/*
 InputKeys=ToolTip
