Fixed Editor localization gathering from plugins that have their own LocRes data

diff --git a/Engine/Config/Localization/Editor.ini b/Engine/Config/Localization/Editor.ini
--- a/Engine/Config/Localization/Editor.ini
+++ b/Engine/Config/Localization/Editor.ini
@@ -19,6 +19,8 @@ CulturesToGenerate=zh-CN
 CommandletClass=GatherTextFromSource
 SearchDirectoryPaths=Source/Editor/
 SearchDirectoryPaths=Plugins/
+ExcludePathFilters=Plugins/Online/*
+ExcludePathFilters=Plugins/Runtime/HTTPChunkInstaller/*
 ExcludePathFilters=*/NoRedist/*
 ExcludePathFilters=*/NotForLicensees/*
 FileNameFilters=*.cpp
