Update UE4 translations to use zh-Hans for Simplified Chinese

#rnx

diff --git a/Engine/Config/Localization/Category.ini b/Engine/Config/Localization/Category.ini
--- a/Engine/Config/Localization/Category.ini
+++ b/Engine/Config/Localization/Category.ini
@@ -12,7 +12,7 @@ CulturesToGenerate=en
 CulturesToGenerate=ja
 CulturesToGenerate=ko
 CulturesToGenerate=pt-BR
-CulturesToGenerate=zh-CN
+CulturesToGenerate=zh-Hans
 
 ;Gather text from metadata
 [GatherTextStep0]
diff --git a/Engine/Config/Localization/Category_Compile.ini b/Engine/Config/Localization/Category_Compile.ini
deleted file mode 100644
diff --git a/Engine/Config/Localization/Editor.ini b/Engine/Config/Localization/Editor.ini
--- a/Engine/Config/Localization/Editor.ini
+++ b/Engine/Config/Localization/Editor.ini
@@ -13,7 +13,7 @@ CulturesToGenerate=es
 CulturesToGenerate=ja
 CulturesToGenerate=ko
 CulturesToGenerate=pt-BR
-CulturesToGenerate=zh-CN
+CulturesToGenerate=zh-Hans
 
 ;Gather text from source code
 [GatherTextStep0]
diff --git a/Engine/Config/Localization/EditorTutorials.ini b/Engine/Config/Localization/EditorTutorials.ini
--- a/Engine/Config/Localization/EditorTutorials.ini
+++ b/Engine/Config/Localization/EditorTutorials.ini
@@ -18,7 +18,7 @@ CulturesToGenerate=en
 CulturesToGenerate=ja
 CulturesToGenerate=ko
 CulturesToGenerate=pt-BR
-CulturesToGenerate=zh-CN
+CulturesToGenerate=zh-Hans
 ModulesToPreload=IntroTutorials
 
 ;Gather text from assets
diff --git a/Engine/Config/Localization/Engine.ini b/Engine/Config/Localization/Engine.ini
--- a/Engine/Config/Localization/Engine.ini
+++ b/Engine/Config/Localization/Engine.ini
@@ -12,7 +12,7 @@ CulturesToGenerate=es
 CulturesToGenerate=ja
 CulturesToGenerate=ko
 CulturesToGenerate=pt-BR
-CulturesToGenerate=zh-CN
+CulturesToGenerate=zh-Hans
 
 ;Gather text from source code
 [GatherTextStep0]
diff --git a/Engine/Config/Localization/Keywords.ini b/Engine/Config/Localization/Keywords.ini
--- a/Engine/Config/Localization/Keywords.ini
+++ b/Engine/Config/Localization/Keywords.ini
@@ -12,7 +12,7 @@ CulturesToGenerate=en
 CulturesToGenerate=ja
 CulturesToGenerate=ko
 CulturesToGenerate=pt-BR
-CulturesToGenerate=zh-CN
+CulturesToGenerate=zh-Hans
 
 ;Gather text from metadata
 [GatherTextStep0]
diff --git a/Engine/Config/Localization/Keywords_Compile.ini b/Engine/Config/Localization/Keywords_Compile.ini
deleted file mode 100644
diff --git a/Engine/Config/Localization/PortableObjectExport.ini b/Engine/Config/Localization/PortableObjectExport.ini
--- a/Engine/Config/Localization/PortableObjectExport.ini
+++ b/Engine/Config/Localization/PortableObjectExport.ini
@@ -7,7 +7,7 @@ CulturesToGenerate=es
 CulturesToGenerate=ja
 CulturesToGenerate=ko
 CulturesToGenerate=pt-BR
-CulturesToGenerate=zh-CN
+CulturesToGenerate=zh-Hans
 
 [GatherTextStep0]
 CommandletClass=InternationalizationExport
diff --git a/Engine/Config/Localization/PortableObjectImport.ini b/Engine/Config/Localization/PortableObjectImport.ini
--- a/Engine/Config/Localization/PortableObjectImport.ini
+++ b/Engine/Config/Localization/PortableObjectImport.ini
@@ -7,7 +7,7 @@ CulturesToGenerate=es
 CulturesToGenerate=ja
 CulturesToGenerate=ko
 CulturesToGenerate=pt-BR
-CulturesToGenerate=zh-CN
+CulturesToGenerate=zh-Hans
 
 [GatherTextStep0]
 CommandletClass=InternationalizationExport
diff --git a/Engine/Config/Localization/PropertyNames.ini b/Engine/Config/Localization/PropertyNames.ini
--- a/Engine/Config/Localization/PropertyNames.ini
+++ b/Engine/Config/Localization/PropertyNames.ini
@@ -12,7 +12,7 @@ CulturesToGenerate=en
 CulturesToGenerate=ja
 CulturesToGenerate=ko
 CulturesToGenerate=pt-BR
-CulturesToGenerate=zh-CN
+CulturesToGenerate=zh-Hans
 
 ;Gather text from metadata
 [GatherTextStep0]
diff --git a/Engine/Config/Localization/RepairData.ini b/Engine/Config/Localization/RepairData.ini
--- a/Engine/Config/Localization/RepairData.ini
+++ b/Engine/Config/Localization/RepairData.ini
@@ -5,7 +5,7 @@ CulturesToGenerate=es
 CulturesToGenerate=ja
 CulturesToGenerate=ko
 CulturesToGenerate=pt-BR
-CulturesToGenerate=zh-CN
+CulturesToGenerate=zh-Hans
 
 ;Repair localization data for Engine
 [GatherTextStep0]
diff --git a/Engine/Config/Localization/ToolTips.ini b/Engine/Config/Localization/ToolTips.ini
--- a/Engine/Config/Localization/ToolTips.ini
+++ b/Engine/Config/Localization/ToolTips.ini
@@ -13,7 +13,7 @@ CulturesToGenerate=en
 CulturesToGenerate=ja
 CulturesToGenerate=ko
 CulturesToGenerate=pt-BR
-CulturesToGenerate=zh-CN
+CulturesToGenerate=zh-Hans
 
 ;Gather text from metadata
 [GatherTextStep0]
diff --git a/Engine/Config/Localization/ToolTips_Compile.ini b/Engine/Config/Localization/ToolTips_Compile.ini
deleted file mode 100644
diff --git a/Engine/Config/Localization/WordCount.ini b/Engine/Config/Localization/WordCount.ini
--- a/Engine/Config/Localization/WordCount.ini
+++ b/Engine/Config/Localization/WordCount.ini
@@ -6,7 +6,7 @@ CulturesToGenerate=es
 CulturesToGenerate=ja
 CulturesToGenerate=ko
 CulturesToGenerate=pt-BR
-CulturesToGenerate=zh-CN
+CulturesToGenerate=zh-Hans
 
 ;Write Text Localization Report for Engine
 [GatherTextStep0]
