Enable collapsing for identical strings in the categories localization target

#rnx

diff --git a/Engine/Config/Localization/Category.ini b/Engine/Config/Localization/Category.ini
--- a/Engine/Config/Localization/Category.ini
+++ b/Engine/Config/Localization/Category.ini
@@ -41,6 +41,7 @@ bPurgeOldEmptyEntries=true
 [GatherTextStep3]
 CommandletClass=InternationalizationExport
 bImportLoc=true
+LocalizedTextCollapseMode=ELocalizedTextCollapseMode::IdenticalNamespaceAndSource
 
 ;Write Localized Text Resource
 [GatherTextStep4]
@@ -51,6 +52,7 @@ ResourceName=Category.locres
 [GatherTextStep5]
 CommandletClass=InternationalizationExport
 bExportLoc=true
+LocalizedTextCollapseMode=ELocalizedTextCollapseMode::IdenticalNamespaceAndSource
 
 ;Write Text Localization Report
 [GatherTextStep6]
