Gather CrashReportClient into the Engine LocRes

#rnx

diff --git a/Engine/Config/Localization/Engine.ini b/Engine/Config/Localization/Engine.ini
--- a/Engine/Config/Localization/Engine.ini
+++ b/Engine/Config/Localization/Engine.ini
@@ -21,6 +21,7 @@ SearchDirectoryPaths=Source/Runtime/
 SearchDirectoryPaths=Source/Developer/
 SearchDirectoryPaths=Plugins/
 SearchDirectoryPaths=Config/
+SearchDirectoryPaths=Source/Programs/CrashReportClient/
 ExcludePathFilters=Config/Localization/*
 ExcludePathFilters=Source/Runtime/Online/BuildPatchServices/*
 ExcludePathFilters=Plugins/*Editor*/*
