[ARPG] Potions logic rework

diff --git a/Engine/Config/Localization/Category_Gather.ini b/Engine/Config/Localization/Category_Gather.ini
new file mode 100644
--- /dev/null
+++ b/Engine/Config/Localization/Category_Gather.ini
@@ -0,0 +1,38 @@
+[CommonSettings]
+SourcePath=Content/Localization/Category
+DestinationPath=Content/Localization/Category
+ManifestName=Category.manifest
+ArchiveName=Category.archive
+NativeCulture=en
+CulturesToGenerate=en
+CulturesToGenerate=es
+CulturesToGenerate=ja
+CulturesToGenerate=ko
+CulturesToGenerate=pt-BR
+CulturesToGenerate=zh-CN
+
+[GatherTextStep0]
+CommandletClass=GatherTextFromMetadata
+IncludePathFilters=Source/Editor/*
+IncludePathFilters=Source/Runtime/*
+IncludePathFilters=Source/Developer/*
+ExcludePathFilters=Source/Developer/NoRedist/CommunityPortalServices/*
+InputKeys=Category
+OutputNamespaces=UObjectCategory
+OutputKeys="\"{FieldPath}\""
+ShouldGatherFromEditorOnlyData=true
+
+[GatherTextStep1]
+CommandletClass=GenerateGatherManifest
+
+[GatherTextStep2]
+CommandletClass=GenerateGatherArchive
+
+[GatherTextStep3]
+CommandletClass=GenerateTextLocalizationReport
+bWordCountReport=true
+WordCountReportName=Category.csv
+bConflictReport=true
+ConflictReportName=Category_Conflicts.txt
+
+
diff --git a/Engine/Config/Localization/EditorTutorials_Gather.ini b/Engine/Config/Localization/EditorTutorials_Gather.ini
new file mode 100644
--- /dev/null
+++ b/Engine/Config/Localization/EditorTutorials_Gather.ini
@@ -0,0 +1,39 @@
+[CommonSettings]
+ManifestDependencies=Content/Localization/Engine/Engine.manifest
+ManifestDependencies=Content/Localization/Editor/Editor.manifest
+ModulesToPreload=IntroTutorials
+SourcePath=Content/Localization/EditorTutorials
+DestinationPath=Content/Localization/EditorTutorials
+ManifestName=EditorTutorials.manifest
+ArchiveName=EditorTutorials.archive
+NativeCulture=en
+CulturesToGenerate=en
+CulturesToGenerate=es
+CulturesToGenerate=ja
+CulturesToGenerate=ko
+CulturesToGenerate=pt-BR
+CulturesToGenerate=zh-CN
+
+[GatherTextStep0]
+CommandletClass=GatherTextFromAssets
+IncludePathFilters=Content/Tutorial/*
+ExcludePathFilters=Content/Localization/*
+PackageFileNameFilters=*.umap
+PackageFileNameFilters=*.uasset
+ShouldGatherFromEditorOnlyData=true
+SkipGatherCache=false
+
+[GatherTextStep1]
+CommandletClass=GenerateGatherManifest
+
+[GatherTextStep2]
+CommandletClass=GenerateGatherArchive
+
+[GatherTextStep3]
+CommandletClass=GenerateTextLocalizationReport
+bWordCountReport=true
+WordCountReportName=EditorTutorials.csv
+bConflictReport=true
+ConflictReportName=EditorTutorials_Conflicts.txt
+
+
diff --git a/Engine/Config/Localization/Editor_Gather.ini b/Engine/Config/Localization/Editor_Gather.ini
new file mode 100644
--- /dev/null
+++ b/Engine/Config/Localization/Editor_Gather.ini
@@ -0,0 +1,49 @@
+[CommonSettings]
+ManifestDependencies=Content/Localization/Engine/Engine.manifest
+SourcePath=Content/Localization/Editor
+DestinationPath=Content/Localization/Editor
+ManifestName=Editor.manifest
+ArchiveName=Editor.archive
+NativeCulture=en
+CulturesToGenerate=en
+CulturesToGenerate=es
+CulturesToGenerate=ja
+CulturesToGenerate=ko
+CulturesToGenerate=pt-BR
+CulturesToGenerate=zh-CN
+
+[GatherTextStep0]
+CommandletClass=GatherTextFromSource
+SearchDirectoryPaths=Source/Editor/
+ExcludePathFilters=Config/Localization/*
+FileNameFilters=*.cpp
+FileNameFilters=*.h
+FileNameFilters=*.c
+FileNameFilters=*.inl
+FileNameFilters=*.mm
+ShouldGatherFromEditorOnlyData=false
+
+[GatherTextStep1]
+CommandletClass=GatherTextFromAssets
+IncludePathFilters=Content/Editor/*
+IncludePathFilters=Content/Editor*
+ExcludePathFilters=Content/Localization/*
+PackageFileNameFilters=*.umap
+PackageFileNameFilters=*.uasset
+ShouldGatherFromEditorOnlyData=true
+SkipGatherCache=false
+
+[GatherTextStep2]
+CommandletClass=GenerateGatherManifest
+
+[GatherTextStep3]
+CommandletClass=GenerateGatherArchive
+
+[GatherTextStep4]
+CommandletClass=GenerateTextLocalizationReport
+bWordCountReport=true
+WordCountReportName=Editor.csv
+bConflictReport=true
+ConflictReportName=Editor_Conflicts.txt
+
+
diff --git a/Engine/Config/Localization/Engine_Gather.ini b/Engine/Config/Localization/Engine_Gather.ini
new file mode 100644
--- /dev/null
+++ b/Engine/Config/Localization/Engine_Gather.ini
@@ -0,0 +1,60 @@
+[CommonSettings]
+SourcePath=Content/Localization/Engine
+DestinationPath=Content/Localization/Engine
+ManifestName=Engine.manifest
+ArchiveName=Engine.archive
+NativeCulture=en
+CulturesToGenerate=en
+CulturesToGenerate=es
+CulturesToGenerate=ja
+CulturesToGenerate=ko
+CulturesToGenerate=pt-BR
+CulturesToGenerate=zh-CN
+
+[GatherTextStep0]
+CommandletClass=GatherTextFromSource
+SearchDirectoryPaths=Source/Runtime/
+SearchDirectoryPaths=Source/Developer/
+SearchDirectoryPaths=Config/
+ExcludePathFilters=Config/Localization/*
+ExcludePathFilters=Source/Developer/NoRedist/CommunityPortalServices/*
+ExcludePathFilters=Source/Developer/NoRedist/UnrealEngineLauncherServices/*
+ExcludePathFilters=Source/Developer/NoRedist/BuildPatchServices/*
+FileNameFilters=*.cpp
+FileNameFilters=*.h
+FileNameFilters=*.c
+FileNameFilters=*.inl
+FileNameFilters=*.mm
+FileNameFilters=*.ini
+ShouldGatherFromEditorOnlyData=false
+
+[GatherTextStep1]
+CommandletClass=GatherTextFromAssets
+IncludePathFilters=Content/*
+ExcludePathFilters=Content/Localization/*
+ExcludePathFilters=Content/Editor/*
+ExcludePathFilters=Content/Tutorial/*
+ExcludePathFilters=Content/Developers/*
+ExcludePathFilters=Content/TestPackages/*
+ExcludePathFilters=Content/QA_Assets/*
+ExcludePathFilters=Content/Maps/Automation/*
+ExcludePathFilters=Content/EngineSounds/*
+PackageFileNameFilters=*.umap
+PackageFileNameFilters=*.uasset
+ShouldGatherFromEditorOnlyData=true
+SkipGatherCache=false
+
+[GatherTextStep2]
+CommandletClass=GenerateGatherManifest
+
+[GatherTextStep3]
+CommandletClass=GenerateGatherArchive
+
+[GatherTextStep4]
+CommandletClass=GenerateTextLocalizationReport
+bWordCountReport=true
+WordCountReportName=Engine.csv
+bConflictReport=true
+ConflictReportName=Engine_Conflicts.txt
+
+
diff --git a/Engine/Config/Localization/Keywords_Gather.ini b/Engine/Config/Localization/Keywords_Gather.ini
new file mode 100644
--- /dev/null
+++ b/Engine/Config/Localization/Keywords_Gather.ini
@@ -0,0 +1,40 @@
+[CommonSettings]
+ManifestDependencies=Content/Localization/Engine/Engine.manifest
+ManifestDependencies=Content/Localization/Editor/Editor.manifest
+SourcePath=Content/Localization/Keywords
+DestinationPath=Content/Localization/Keywords
+ManifestName=Keywords.manifest
+ArchiveName=Keywords.archive
+NativeCulture=en
+CulturesToGenerate=en
+CulturesToGenerate=es
+CulturesToGenerate=ja
+CulturesToGenerate=ko
+CulturesToGenerate=pt-BR
+CulturesToGenerate=zh-CN
+
+[GatherTextStep0]
+CommandletClass=GatherTextFromMetadata
+IncludePathFilters=Source/Editor/*
+IncludePathFilters=Source/Runtime/*
+IncludePathFilters=Source/Developer/*
+ExcludePathFilters=Source/Developer/NoRedist/CommunityPortalServices/*
+InputKeys=Keywords
+OutputNamespaces=UObjectKeywords
+OutputKeys="\"{FieldPath}\""
+ShouldGatherFromEditorOnlyData=true
+
+[GatherTextStep1]
+CommandletClass=GenerateGatherManifest
+
+[GatherTextStep2]
+CommandletClass=GenerateGatherArchive
+
+[GatherTextStep3]
+CommandletClass=GenerateTextLocalizationReport
+bWordCountReport=true
+WordCountReportName=Keywords.csv
+bConflictReport=true
+ConflictReportName=Keywords_Conflicts.txt
+
+
diff --git a/Engine/Config/Localization/PropertyNames_Gather.ini b/Engine/Config/Localization/PropertyNames_Gather.ini
new file mode 100644
--- /dev/null
+++ b/Engine/Config/Localization/PropertyNames_Gather.ini
@@ -0,0 +1,40 @@
+[CommonSettings]
+ManifestDependencies=Content/Localization/Engine/Engine.manifest
+ManifestDependencies=Content/Localization/Editor/Editor.manifest
+SourcePath=Content/Localization/PropertyNames
+DestinationPath=Content/Localization/PropertyNames
+ManifestName=PropertyNames.manifest
+ArchiveName=PropertyNames.archive
+NativeCulture=en
+CulturesToGenerate=en
+CulturesToGenerate=es
+CulturesToGenerate=ja
+CulturesToGenerate=ko
+CulturesToGenerate=pt-BR
+CulturesToGenerate=zh-CN
+
+[GatherTextStep0]
+CommandletClass=GatherTextFromMetadata
+IncludePathFilters=Source/Editor/*
+IncludePathFilters=Source/Runtime/*
+IncludePathFilters=Source/Developer/*
+ExcludePathFilters=Source/Developer/NoRedist/CommunityPortalServices/*
+InputKeys=DisplayName
+OutputNamespaces=UObjectDisplayNames
+OutputKeys="\"{FieldPath}\""
+ShouldGatherFromEditorOnlyData=true
+
+[GatherTextStep1]
+CommandletClass=GenerateGatherManifest
+
+[GatherTextStep2]
+CommandletClass=GenerateGatherArchive
+
+[GatherTextStep3]
+CommandletClass=GenerateTextLocalizationReport
+bWordCountReport=true
+WordCountReportName=PropertyNames.csv
+bConflictReport=true
+ConflictReportName=PropertyNames_Conflicts.txt
+
+
diff --git a/Engine/Config/Localization/ToolTips_Gather.ini b/Engine/Config/Localization/ToolTips_Gather.ini
new file mode 100644
--- /dev/null
+++ b/Engine/Config/Localization/ToolTips_Gather.ini
@@ -0,0 +1,43 @@
+[CommonSettings]
+ManifestDependencies=Content/Localization/Engine/Engine.manifest
+ManifestDependencies=Content/Localization/Editor/Editor.manifest
+SourcePath=Content/Localization/ToolTips
+DestinationPath=Content/Localization/ToolTips
+ManifestName=ToolTips.manifest
+ArchiveName=ToolTips.archive
+NativeCulture=en
+CulturesToGenerate=en
+CulturesToGenerate=es
+CulturesToGenerate=ja
+CulturesToGenerate=ko
+CulturesToGenerate=pt-BR
+CulturesToGenerate=zh-CN
+
+[GatherTextStep0]
+CommandletClass=GatherTextFromMetadata
+IncludePathFilters=Source/Editor/*
+IncludePathFilters=Source/Runtime/*
+IncludePathFilters=Source/Developer/*
+ExcludePathFilters=Source/Developer/NoRedist/CommunityPortalServices/*
+InputKeys=ToolTip
+InputKeys=ShortToolTip
+OutputNamespaces=UObjectToolTips
+OutputNamespaces=UObjectShortToolTips
+OutputKeys="\"{FieldPath}\""
+OutputKeys="\"{FieldPath}\""
+ShouldGatherFromEditorOnlyData=true
+
+[GatherTextStep1]
+CommandletClass=GenerateGatherManifest
+
+[GatherTextStep2]
+CommandletClass=GenerateGatherArchive
+
+[GatherTextStep3]
+CommandletClass=GenerateTextLocalizationReport
+bWordCountReport=true
+WordCountReportName=ToolTips.csv
+bConflictReport=true
+ConflictReportName=ToolTips_Conflicts.txt
+
+
