Removed Arabic from Engine localization targets

This only existed to specify the legacy localization fallback font name to use for Arabic, and since that feature was removed in 4.24 Arabic is no longer needed as we don't translate the Engine itself into that language.

#rnx

diff --git a/Engine/Config/Localization/Category.ini b/Engine/Config/Localization/Category.ini
--- a/Engine/Config/Localization/Category.ini
+++ b/Engine/Config/Localization/Category.ini
@@ -6,7 +6,6 @@ ArchiveName=Category.archive
 PortableObjectName=Category.po
 PlatformSplitMode=ELocTextPlatformSplitMode::Restricted
 NativeCulture=en
-CulturesToGenerate=ar
 CulturesToGenerate=es
 CulturesToGenerate=en
 CulturesToGenerate=ja
diff --git a/Engine/Config/Localization/Editor.ini b/Engine/Config/Localization/Editor.ini
--- a/Engine/Config/Localization/Editor.ini
+++ b/Engine/Config/Localization/Editor.ini
@@ -7,7 +7,6 @@ ArchiveName=Editor.archive
 PortableObjectName=Editor.po
 PlatformSplitMode=ELocTextPlatformSplitMode::Restricted
 NativeCulture=en
-CulturesToGenerate=ar
 CulturesToGenerate=en
 CulturesToGenerate=es
 CulturesToGenerate=ja
diff --git a/Engine/Config/Localization/EditorTutorials.ini b/Engine/Config/Localization/EditorTutorials.ini
--- a/Engine/Config/Localization/EditorTutorials.ini
+++ b/Engine/Config/Localization/EditorTutorials.ini
@@ -12,7 +12,6 @@ ArchiveName=EditorTutorials.archive
 PortableObjectName=EditorTutorials.po
 PlatformSplitMode=ELocTextPlatformSplitMode::Restricted
 NativeCulture=en
-CulturesToGenerate=ar
 CulturesToGenerate=es
 CulturesToGenerate=en
 CulturesToGenerate=ja
diff --git a/Engine/Config/Localization/Engine.ini b/Engine/Config/Localization/Engine.ini
--- a/Engine/Config/Localization/Engine.ini
+++ b/Engine/Config/Localization/Engine.ini
@@ -6,7 +6,6 @@ ArchiveName=Engine.archive
 PortableObjectName=Engine.po
 PlatformSplitMode=ELocTextPlatformSplitMode::Restricted
 NativeCulture=en
-CulturesToGenerate=ar
 CulturesToGenerate=en
 CulturesToGenerate=es
 CulturesToGenerate=ja
diff --git a/Engine/Config/Localization/Keywords.ini b/Engine/Config/Localization/Keywords.ini
--- a/Engine/Config/Localization/Keywords.ini
+++ b/Engine/Config/Localization/Keywords.ini
@@ -6,7 +6,6 @@ ArchiveName=Keywords.archive
 PortableObjectName=Keywords.po
 PlatformSplitMode=ELocTextPlatformSplitMode::Restricted
 NativeCulture=en
-CulturesToGenerate=ar
 CulturesToGenerate=es
 CulturesToGenerate=en
 CulturesToGenerate=ja
diff --git a/Engine/Config/Localization/PortableObjectExport.ini b/Engine/Config/Localization/PortableObjectExport.ini
--- a/Engine/Config/Localization/PortableObjectExport.ini
+++ b/Engine/Config/Localization/PortableObjectExport.ini
@@ -1,7 +1,6 @@
 [CommonSettings]
 bExportLoc=true
 NativeCulture=en
-CulturesToGenerate=ar
 CulturesToGenerate=en
 CulturesToGenerate=es
 CulturesToGenerate=ja
diff --git a/Engine/Config/Localization/PortableObjectImport.ini b/Engine/Config/Localization/PortableObjectImport.ini
--- a/Engine/Config/Localization/PortableObjectImport.ini
+++ b/Engine/Config/Localization/PortableObjectImport.ini
@@ -1,7 +1,6 @@
 [CommonSettings]
 bImportLoc=true
 NativeCulture=en
-CulturesToGenerate=ar
 CulturesToGenerate=en
 CulturesToGenerate=es
 CulturesToGenerate=ja
diff --git a/Engine/Config/Localization/PropertyNames.ini b/Engine/Config/Localization/PropertyNames.ini
--- a/Engine/Config/Localization/PropertyNames.ini
+++ b/Engine/Config/Localization/PropertyNames.ini
@@ -6,7 +6,6 @@ ArchiveName=PropertyNames.archive
 PortableObjectName=PropertyNames.po
 PlatformSplitMode=ELocTextPlatformSplitMode::Restricted
 NativeCulture=en
-CulturesToGenerate=ar
 CulturesToGenerate=es
 CulturesToGenerate=en
 CulturesToGenerate=ja
diff --git a/Engine/Config/Localization/RepairData.ini b/Engine/Config/Localization/RepairData.ini
--- a/Engine/Config/Localization/RepairData.ini
+++ b/Engine/Config/Localization/RepairData.ini
@@ -1,5 +1,4 @@
 [CommonSettings]
-CulturesToGenerate=ar
 CulturesToGenerate=en
 CulturesToGenerate=es
 CulturesToGenerate=ja
diff --git a/Engine/Config/Localization/ToolTips.ini b/Engine/Config/Localization/ToolTips.ini
--- a/Engine/Config/Localization/ToolTips.ini
+++ b/Engine/Config/Localization/ToolTips.ini
@@ -7,7 +7,6 @@ ArchiveName=ToolTips.archive
 PortableObjectName=ToolTips.po
 PlatformSplitMode=ELocTextPlatformSplitMode::Restricted
 NativeCulture=en
-CulturesToGenerate=ar
 CulturesToGenerate=es
 CulturesToGenerate=en
 CulturesToGenerate=ja
diff --git a/Engine/Config/Localization/WordCount.ini b/Engine/Config/Localization/WordCount.ini
--- a/Engine/Config/Localization/WordCount.ini
+++ b/Engine/Config/Localization/WordCount.ini
@@ -1,7 +1,6 @@
 [CommonSettings]
 bWordCountReport=true
 DestinationPath=\\epicgames.net\root\UE3\Localization\WordCounts
-CulturesToGenerate=ar
 CulturesToGenerate=es
 CulturesToGenerate=ja
 CulturesToGenerate=ko
