Gather landscape and mesh paint mode panel meta-data into the Editor localization target

#rnx

diff --git a/Engine/Config/Localization/Category.ini b/Engine/Config/Localization/Category.ini
--- a/Engine/Config/Localization/Category.ini
+++ b/Engine/Config/Localization/Category.ini
@@ -18,6 +18,9 @@ CulturesToGenerate=zh-Hans
 ;Gather text from metadata
 [GatherTextStep0]
 CommandletClass=GatherTextFromMetaData
+ModulesToPreload=StructViewer
+ModulesToPreload=MeshPaintMode
+ModulesToPreload=LandscapeEditor
 IncludePathFilters=Source/Editor/*
 IncludePathFilters=Source/Runtime/*
 IncludePathFilters=Source/Developer/*
diff --git a/Engine/Config/Localization/Editor.ini b/Engine/Config/Localization/Editor.ini
--- a/Engine/Config/Localization/Editor.ini
+++ b/Engine/Config/Localization/Editor.ini
@@ -45,6 +45,8 @@ ShouldGatherFromEditorOnlyData=true
 [GatherTextStep2]
 CommandletClass=GatherTextFromMetaData
 ModulesToPreload=StructViewer
+ModulesToPreload=MeshPaintMode
+ModulesToPreload=LandscapeEditor
 IncludePathFilters=Source/Editor/*
 IncludePathFilters=Source/Runtime/*
 IncludePathFilters=Source/Developer/*
@@ -64,6 +66,8 @@ FieldOwnerTypesToInclude=InternationalizationSettingsModel
 FieldOwnerTypesToInclude=NavigationSystemV1
 FieldOwnerTypesToInclude=PersonaOptions
 FieldOwnerTypesToInclude=RecastNavMesh
+; modes panel
+FieldOwnerTypesToInclude=*Landscape*
 InputKeys=DisplayName
 OutputNamespaces=UObjectDisplayNames
 OutputKeys="{FieldPath}"
diff --git a/Engine/Config/Localization/PropertyNames.ini b/Engine/Config/Localization/PropertyNames.ini
--- a/Engine/Config/Localization/PropertyNames.ini
+++ b/Engine/Config/Localization/PropertyNames.ini
@@ -18,6 +18,9 @@ CulturesToGenerate=zh-Hans
 ;Gather text from metadata
 [GatherTextStep0]
 CommandletClass=GatherTextFromMetaData
+ModulesToPreload=StructViewer
+ModulesToPreload=MeshPaintMode
+ModulesToPreload=LandscapeEditor
 IncludePathFilters=Source/Editor/*
 IncludePathFilters=Source/Runtime/*
 IncludePathFilters=Source/Developer/*
diff --git a/Engine/Config/Localization/ToolTips.ini b/Engine/Config/Localization/ToolTips.ini
--- a/Engine/Config/Localization/ToolTips.ini
+++ b/Engine/Config/Localization/ToolTips.ini
@@ -18,6 +18,9 @@ CulturesToGenerate=zh-Hans
 ;Gather text from metadata
 [GatherTextStep0]
 CommandletClass=GatherTextFromMetaData
+ModulesToPreload=StructViewer
+ModulesToPreload=MeshPaintMode
+ModulesToPreload=LandscapeEditor
 IncludePathFilters=Source/Editor/*
 IncludePathFilters=Source/Runtime/*
 IncludePathFilters=Source/Developer/*
