Update Engine localization config files to also support German (de), Latin American Spanish (es-419) and French (fr) for localization commandlet.
#rnx

diff --git a/Engine/Config/Localization/Category.ini b/Engine/Config/Localization/Category.ini
--- a/Engine/Config/Localization/Category.ini
+++ b/Engine/Config/Localization/Category.ini
@@ -6,8 +6,11 @@ ManifestName=Category.manifest
 ArchiveName=Category.archive
 PortableObjectName=Category.po
 NativeCulture=en
-CulturesToGenerate=es
 CulturesToGenerate=en
+CulturesToGenerate=de
+CulturesToGenerate=es
+CulturesToGenerate=es-419
+CulturesToGenerate=fr
 CulturesToGenerate=ja
 CulturesToGenerate=ko
 CulturesToGenerate=pt-BR
diff --git a/Engine/Config/Localization/Editor.ini b/Engine/Config/Localization/Editor.ini
--- a/Engine/Config/Localization/Editor.ini
+++ b/Engine/Config/Localization/Editor.ini
@@ -7,7 +7,10 @@ ArchiveName=Editor.archive
 PortableObjectName=Editor.po
 NativeCulture=en
 CulturesToGenerate=en
+CulturesToGenerate=de
 CulturesToGenerate=es
+CulturesToGenerate=es-419
+CulturesToGenerate=fr
 CulturesToGenerate=ja
 CulturesToGenerate=ko
 CulturesToGenerate=pt-BR
diff --git a/Engine/Config/Localization/EditorTutorials.ini b/Engine/Config/Localization/EditorTutorials.ini
--- a/Engine/Config/Localization/EditorTutorials.ini
+++ b/Engine/Config/Localization/EditorTutorials.ini
@@ -11,8 +11,11 @@ ManifestName=EditorTutorials.manifest
 ArchiveName=EditorTutorials.archive
 PortableObjectName=EditorTutorials.po
 NativeCulture=en
-CulturesToGenerate=es
 CulturesToGenerate=en
+CulturesToGenerate=de
+CulturesToGenerate=es
+CulturesToGenerate=es-419
+CulturesToGenerate=fr
 CulturesToGenerate=ja
 CulturesToGenerate=ko
 CulturesToGenerate=pt-BR
diff --git a/Engine/Config/Localization/Engine.ini b/Engine/Config/Localization/Engine.ini
--- a/Engine/Config/Localization/Engine.ini
+++ b/Engine/Config/Localization/Engine.ini
@@ -6,7 +6,10 @@ ArchiveName=Engine.archive
 PortableObjectName=Engine.po
 NativeCulture=en
 CulturesToGenerate=en
+CulturesToGenerate=de
 CulturesToGenerate=es
+CulturesToGenerate=es-419
+CulturesToGenerate=fr
 CulturesToGenerate=ja
 CulturesToGenerate=ko
 CulturesToGenerate=pt-BR
diff --git a/Engine/Config/Localization/Keywords.ini b/Engine/Config/Localization/Keywords.ini
--- a/Engine/Config/Localization/Keywords.ini
+++ b/Engine/Config/Localization/Keywords.ini
@@ -5,8 +5,11 @@ ManifestName=Keywords.manifest
 ArchiveName=Keywords.archive
 PortableObjectName=Keywords.po
 NativeCulture=en
-CulturesToGenerate=es
 CulturesToGenerate=en
+CulturesToGenerate=de
+CulturesToGenerate=es
+CulturesToGenerate=es-419
+CulturesToGenerate=fr
 CulturesToGenerate=ja
 CulturesToGenerate=ko
 CulturesToGenerate=pt-BR
diff --git a/Engine/Config/Localization/PropertyNames.ini b/Engine/Config/Localization/PropertyNames.ini
--- a/Engine/Config/Localization/PropertyNames.ini
+++ b/Engine/Config/Localization/PropertyNames.ini
@@ -6,8 +6,11 @@ ManifestName=PropertyNames.manifest
 ArchiveName=PropertyNames.archive
 PortableObjectName=PropertyNames.po
 NativeCulture=en
-CulturesToGenerate=es
 CulturesToGenerate=en
+CulturesToGenerate=de
+CulturesToGenerate=es
+CulturesToGenerate=es-419
+CulturesToGenerate=fr
 CulturesToGenerate=ja
 CulturesToGenerate=ko
 CulturesToGenerate=pt-BR
diff --git a/Engine/Config/Localization/ToolTips.ini b/Engine/Config/Localization/ToolTips.ini
--- a/Engine/Config/Localization/ToolTips.ini
+++ b/Engine/Config/Localization/ToolTips.ini
@@ -6,8 +6,11 @@ ManifestName=ToolTips.manifest
 ArchiveName=ToolTips.archive
 PortableObjectName=ToolTips.po
 NativeCulture=en
-CulturesToGenerate=es
 CulturesToGenerate=en
+CulturesToGenerate=de
+CulturesToGenerate=es
+CulturesToGenerate=es-419
+CulturesToGenerate=fr
 CulturesToGenerate=ja
 CulturesToGenerate=ko
 CulturesToGenerate=pt-BR
diff --git a/Engine/Config/Localization/WordCount.ini b/Engine/Config/Localization/WordCount.ini
--- a/Engine/Config/Localization/WordCount.ini
+++ b/Engine/Config/Localization/WordCount.ini
@@ -1,7 +1,10 @@
 [CommonSettings]
 bWordCountReport=true
 DestinationPath=\\epicgames.net\root\UE3\Localization\WordCounts
+CulturesToGenerate=de
 CulturesToGenerate=es
+CulturesToGenerate=es-419
+CulturesToGenerate=fr
 CulturesToGenerate=ja
 CulturesToGenerate=ko
 CulturesToGenerate=pt-BR
