Adding uplugin to loc gather

diff --git a/Engine/Config/Localization/Editor.ini b/Engine/Config/Localization/Editor.ini
--- a/Engine/Config/Localization/Editor.ini
+++ b/Engine/Config/Localization/Editor.ini
@@ -38,8 +38,10 @@ ShouldGatherFromEditorOnlyData=true
 CommandletClass=GatherTextFromAssets
 IncludePathFilters=Content/Editor/*
 IncludePathFilters=Content/Editor*
+IncludePathFilters=Plugins/*
 PackageFileNameFilters=*.umap
 PackageFileNameFilters=*.uasset
+PackageFileNameFilters=*.uplugin
 bFixBroken=false
 ShouldGatherFromEditorOnlyData=true
 
