Removing uplugin asset from loc gather since it is not working

diff --git a/Engine/Config/Localization/Editor.ini b/Engine/Config/Localization/Editor.ini
--- a/Engine/Config/Localization/Editor.ini
+++ b/Engine/Config/Localization/Editor.ini
@@ -38,10 +38,8 @@ ShouldGatherFromEditorOnlyData=true
 CommandletClass=GatherTextFromAssets
 IncludePathFilters=Content/Editor/*
 IncludePathFilters=Content/Editor*
-IncludePathFilters=Plugins/*
 PackageFileNameFilters=*.umap
 PackageFileNameFilters=*.uasset
-PackageFileNameFilters=*.uplugin
 bFixBroken=false
 ShouldGatherFromEditorOnlyData=true
 
