Ensure OutputLog module is loaded when gathering meta-data localization

#rnx

diff --git a/Engine/Config/Localization/Category.ini b/Engine/Config/Localization/Category.ini
--- a/Engine/Config/Localization/Category.ini
+++ b/Engine/Config/Localization/Category.ini
@@ -19,6 +19,7 @@ CulturesToGenerate=zh-Hans
 ;Gather text from metadata
 [GatherTextStep0]
 CommandletClass=GatherTextFromMetaData
+ModulesToPreload=OutputLog
 ModulesToPreload=StructViewer
 ModulesToPreload=MeshPaintMode
 ModulesToPreload=LandscapeEditor
diff --git a/Engine/Config/Localization/PropertyNames.ini b/Engine/Config/Localization/PropertyNames.ini
--- a/Engine/Config/Localization/PropertyNames.ini
+++ b/Engine/Config/Localization/PropertyNames.ini
@@ -19,6 +19,7 @@ CulturesToGenerate=zh-Hans
 ;Gather text from metadata
 [GatherTextStep0]
 CommandletClass=GatherTextFromMetaData
+ModulesToPreload=OutputLog
 ModulesToPreload=StructViewer
 ModulesToPreload=MeshPaintMode
 ModulesToPreload=LandscapeEditor
diff --git a/Engine/Config/Localization/ToolTips.ini b/Engine/Config/Localization/ToolTips.ini
--- a/Engine/Config/Localization/ToolTips.ini
+++ b/Engine/Config/Localization/ToolTips.ini
@@ -19,6 +19,7 @@ CulturesToGenerate=zh-Hans
 ;Gather text from metadata
 [GatherTextStep0]
 CommandletClass=GatherTextFromMetaData
+ModulesToPreload=OutputLog
 ModulesToPreload=StructViewer
 ModulesToPreload=MeshPaintMode
 ModulesToPreload=LandscapeEditor
