Localization:
Adding languages to UE for parity support.

#rnx

diff --git a/Engine/Config/Localization/Category.ini b/Engine/Config/Localization/Category.ini
--- a/Engine/Config/Localization/Category.ini
+++ b/Engine/Config/Localization/Category.ini
@@ -8,12 +8,17 @@ PortableObjectName=Category.po
 NativeCulture=en
 CulturesToGenerate=en
 CulturesToGenerate=de
+CulturesToGenerate=pl
 CulturesToGenerate=es
 CulturesToGenerate=es-419
 CulturesToGenerate=fr
 CulturesToGenerate=ja
+CulturesToGenerate=ar
 CulturesToGenerate=ko
 CulturesToGenerate=pt-BR
+CulturesToGenerate=it
+CulturesToGenerate=ru
+CulturesToGenerate=tr
 CulturesToGenerate=zh-Hans
 
 ;Gather text from metadata
diff --git a/Engine/Config/Localization/Editor.ini b/Engine/Config/Localization/Editor.ini
--- a/Engine/Config/Localization/Editor.ini
+++ b/Engine/Config/Localization/Editor.ini
@@ -8,12 +8,17 @@ PortableObjectName=Editor.po
 NativeCulture=en
 CulturesToGenerate=en
 CulturesToGenerate=de
+CulturesToGenerate=pl
 CulturesToGenerate=es
 CulturesToGenerate=es-419
 CulturesToGenerate=fr
 CulturesToGenerate=ja
+CulturesToGenerate=ar
 CulturesToGenerate=ko
 CulturesToGenerate=pt-BR
+CulturesToGenerate=it
+CulturesToGenerate=ru
+CulturesToGenerate=tr
 CulturesToGenerate=zh-Hans
 
 ;Gather text from source code
diff --git a/Engine/Config/Localization/EditorTutorials.ini b/Engine/Config/Localization/EditorTutorials.ini
--- a/Engine/Config/Localization/EditorTutorials.ini
+++ b/Engine/Config/Localization/EditorTutorials.ini
@@ -13,12 +13,17 @@ PortableObjectName=EditorTutorials.po
 NativeCulture=en
 CulturesToGenerate=en
 CulturesToGenerate=de
+CulturesToGenerate=pl
 CulturesToGenerate=es
 CulturesToGenerate=es-419
 CulturesToGenerate=fr
 CulturesToGenerate=ja
+CulturesToGenerate=ar
 CulturesToGenerate=ko
 CulturesToGenerate=pt-BR
+CulturesToGenerate=it
+CulturesToGenerate=ru
+CulturesToGenerate=tr
 CulturesToGenerate=zh-Hans
 
 ;NOTE: IntroTutorials is now moved into the GuidedTutorial plugin.
diff --git a/Engine/Config/Localization/Engine.ini b/Engine/Config/Localization/Engine.ini
--- a/Engine/Config/Localization/Engine.ini
+++ b/Engine/Config/Localization/Engine.ini
@@ -7,12 +7,17 @@ PortableObjectName=Engine.po
 NativeCulture=en
 CulturesToGenerate=en
 CulturesToGenerate=de
+CulturesToGenerate=pl
 CulturesToGenerate=es
 CulturesToGenerate=es-419
 CulturesToGenerate=fr
 CulturesToGenerate=ja
+CulturesToGenerate=ar
 CulturesToGenerate=ko
 CulturesToGenerate=pt-BR
+CulturesToGenerate=it
+CulturesToGenerate=ru
+CulturesToGenerate=tr
 CulturesToGenerate=zh-Hans
 
 ;Gather text from source code
diff --git a/Engine/Config/Localization/Keywords.ini b/Engine/Config/Localization/Keywords.ini
--- a/Engine/Config/Localization/Keywords.ini
+++ b/Engine/Config/Localization/Keywords.ini
@@ -7,12 +7,17 @@ PortableObjectName=Keywords.po
 NativeCulture=en
 CulturesToGenerate=en
 CulturesToGenerate=de
+CulturesToGenerate=pl
 CulturesToGenerate=es
 CulturesToGenerate=es-419
 CulturesToGenerate=fr
 CulturesToGenerate=ja
+CulturesToGenerate=ar
 CulturesToGenerate=ko
 CulturesToGenerate=pt-BR
+CulturesToGenerate=it
+CulturesToGenerate=ru
+CulturesToGenerate=tr
 CulturesToGenerate=zh-Hans
 
 ;Gather text from metadata
diff --git a/Engine/Config/Localization/PropertyNames.ini b/Engine/Config/Localization/PropertyNames.ini
--- a/Engine/Config/Localization/PropertyNames.ini
+++ b/Engine/Config/Localization/PropertyNames.ini
@@ -8,12 +8,17 @@ PortableObjectName=PropertyNames.po
 NativeCulture=en
 CulturesToGenerate=en
 CulturesToGenerate=de
+CulturesToGenerate=pl
 CulturesToGenerate=es
 CulturesToGenerate=es-419
 CulturesToGenerate=fr
 CulturesToGenerate=ja
+CulturesToGenerate=ar
 CulturesToGenerate=ko
 CulturesToGenerate=pt-BR
+CulturesToGenerate=it
+CulturesToGenerate=ru
+CulturesToGenerate=tr
 CulturesToGenerate=zh-Hans
 
 ;Gather text from metadata
diff --git a/Engine/Config/Localization/ToolTips.ini b/Engine/Config/Localization/ToolTips.ini
--- a/Engine/Config/Localization/ToolTips.ini
+++ b/Engine/Config/Localization/ToolTips.ini
@@ -8,12 +8,17 @@ PortableObjectName=ToolTips.po
 NativeCulture=en
 CulturesToGenerate=en
 CulturesToGenerate=de
+CulturesToGenerate=pl
 CulturesToGenerate=es
 CulturesToGenerate=es-419
 CulturesToGenerate=fr
 CulturesToGenerate=ja
+CulturesToGenerate=ar
 CulturesToGenerate=ko
 CulturesToGenerate=pt-BR
+CulturesToGenerate=it
+CulturesToGenerate=ru
+CulturesToGenerate=tr
 CulturesToGenerate=zh-Hans
 
 ;Gather text from metadata
diff --git a/Engine/Config/Localization/WordCount.ini b/Engine/Config/Localization/WordCount.ini
--- a/Engine/Config/Localization/WordCount.ini
+++ b/Engine/Config/Localization/WordCount.ini
@@ -2,12 +2,17 @@
 bWordCountReport=true
 DestinationPath=\\epicgames.net\root\UE3\Localization\WordCounts
 CulturesToGenerate=de
+CulturesToGenerate=pl
 CulturesToGenerate=es
 CulturesToGenerate=es-419
 CulturesToGenerate=fr
 CulturesToGenerate=ja
+CulturesToGenerate=ar
 CulturesToGenerate=ko
 CulturesToGenerate=pt-BR
+CulturesToGenerate=it
+CulturesToGenerate=ru
+CulturesToGenerate=tr
 CulturesToGenerate=zh-Hans
 
 ;Write Text Localization Report for Engine
