Localization:
Have base Engine localization config files exclude the MetaHuman plugins for localization. MetaHuman plugins will use auto plugin localization.

diff --git a/Engine/Config/Localization/Category.ini b/Engine/Config/Localization/Category.ini
--- a/Engine/Config/Localization/Category.ini
+++ b/Engine/Config/Localization/Category.ini
@@ -35,6 +35,7 @@ IncludePathFilters=Plugins/*
 ExcludePathFilters=*/NoRedist/*
 ExcludePathFilters=*/NotForLicensees/*
 ExcludePathFilters=*/LimitedAccess/*
+ExcludePathFilters=Plugins/*/MetaHuman*
 InputKeys=Category
 OutputNamespaces=UObjectCategory
 OutputKeys="{MetaDataValue}"
diff --git a/Engine/Config/Localization/Editor.ini b/Engine/Config/Localization/Editor.ini
--- a/Engine/Config/Localization/Editor.ini
+++ b/Engine/Config/Localization/Editor.ini
@@ -32,6 +32,7 @@ ExcludePathFilters=*/NoRedist/*
 ExcludePathFilters=*/NotForLicensees/*
 ExcludePathFilters=*/LimitedAccess/*
 ExcludePathFilters=*/ThirdParty/*
+ExcludePathFilters=Plugins/*/MetaHuman*
 FileNameFilters=*.cpp
 FileNameFilters=*.h
 FileNameFilters=*.c
@@ -62,6 +63,7 @@ IncludePathFilters=Plugins/*
 ExcludePathFilters=*/NoRedist/*
 ExcludePathFilters=*/NotForLicensees/*
 ExcludePathFilters=*/LimitedAccess/*
+ExcludePathFilters=Plugins/*/MetaHuman*
 FieldTypesToInclude=Property
 ; general settings
 FieldOwnerTypesToInclude=DeveloperSettings
diff --git a/Engine/Config/Localization/Engine.ini b/Engine/Config/Localization/Engine.ini
--- a/Engine/Config/Localization/Engine.ini
+++ b/Engine/Config/Localization/Engine.ini
@@ -37,6 +37,7 @@ ExcludePathFilters=*/NoRedist/*
 ExcludePathFilters=*/NotForLicensees/*
 ExcludePathFilters=*/LimitedAccess/*
 ExcludePathFilters=*/ThirdParty/*
+ExcludePathFilters=Plugins/*/MetaHuman*
 FileNameFilters=*.cpp
 FileNameFilters=*.h
 FileNameFilters=*.c
diff --git a/Engine/Config/Localization/Keywords.ini b/Engine/Config/Localization/Keywords.ini
--- a/Engine/Config/Localization/Keywords.ini
+++ b/Engine/Config/Localization/Keywords.ini
@@ -30,6 +30,7 @@ IncludePathFilters=Plugins/*
 ExcludePathFilters=*/NoRedist/*
 ExcludePathFilters=*/NotForLicensees/*
 ExcludePathFilters=*/LimitedAccess/*
+ExcludePathFilters=Plugins/*/MetaHuman*
 InputKeys=Keywords
 OutputNamespaces=UObjectKeywords
 OutputKeys="{FieldPath}"
diff --git a/Engine/Config/Localization/PropertyNames.ini b/Engine/Config/Localization/PropertyNames.ini
--- a/Engine/Config/Localization/PropertyNames.ini
+++ b/Engine/Config/Localization/PropertyNames.ini
@@ -35,6 +35,7 @@ IncludePathFilters=Plugins/*
 ExcludePathFilters=*/NoRedist/*
 ExcludePathFilters=*/NotForLicensees/*
 ExcludePathFilters=*/LimitedAccess/*
+ExcludePathFilters=Plugins/*/MetaHuman*
 InputKeys=DisplayName
 OutputNamespaces=UObjectDisplayNames
 OutputKeys="{FieldPath}"
diff --git a/Engine/Config/Localization/ToolTips.ini b/Engine/Config/Localization/ToolTips.ini
--- a/Engine/Config/Localization/ToolTips.ini
+++ b/Engine/Config/Localization/ToolTips.ini
@@ -35,6 +35,7 @@ IncludePathFilters=Plugins/*
 ExcludePathFilters=*/NoRedist/*
 ExcludePathFilters=*/NotForLicensees/*
 ExcludePathFilters=*/LimitedAccess/*
+ExcludePathFilters=Plugins/*/MetaHuman*
 InputKeys=ToolTip
 OutputNamespaces=UObjectToolTips
 OutputKeys="{FieldPath}"
