Update localization exclude filters for MetaHuman plugins to avoid double-gathering

#tests Run "Localize" command for Engine project, observe Engine.po, Category.po, Editor.po, ToolTips.po and PropertyNames.po don't contain MetaHuman strings anymore

diff --git a/Engine/Config/Localization/Category.ini b/Engine/Config/Localization/Category.ini
--- a/Engine/Config/Localization/Category.ini
+++ b/Engine/Config/Localization/Category.ini
@@ -32,6 +32,7 @@ IncludePathFilters=Source/Editor/*
 IncludePathFilters=Source/Runtime/*
 IncludePathFilters=Source/Developer/*
 IncludePathFilters=Plugins/*
+ExcludePathFilters=Plugins/MetaHuman/*
 ExcludePathFilters=*/NoRedist/*
 ExcludePathFilters=*/NotForLicensees/*
 ExcludePathFilters=*/LimitedAccess/*
diff --git a/Engine/Config/Localization/Editor.ini b/Engine/Config/Localization/Editor.ini
--- a/Engine/Config/Localization/Editor.ini
+++ b/Engine/Config/Localization/Editor.ini
@@ -26,6 +26,7 @@ CulturesToGenerate=zh-Hans
 CommandletClass=GatherTextFromSource
 SearchDirectoryPaths=Source/Editor/
 SearchDirectoryPaths=Plugins/
+ExcludePathFilters=Plugins/MetaHuman/*
 ExcludePathFilters=Plugins/Online/*
 ExcludePathFilters=Plugins/Runtime/HTTPChunkInstaller/*
 ExcludePathFilters=*/NoRedist/*
@@ -60,6 +61,7 @@ IncludePathFilters=Source/Editor/*
 IncludePathFilters=Source/Runtime/*
 IncludePathFilters=Source/Developer/*
 IncludePathFilters=Plugins/*
+ExcludePathFilters=Plugins/MetaHuman/*
 ExcludePathFilters=*/NoRedist/*
 ExcludePathFilters=*/NotForLicensees/*
 ExcludePathFilters=*/LimitedAccess/*
diff --git a/Engine/Config/Localization/Engine.ini b/Engine/Config/Localization/Engine.ini
--- a/Engine/Config/Localization/Engine.ini
+++ b/Engine/Config/Localization/Engine.ini
@@ -31,6 +31,7 @@ SearchDirectoryPaths=Source/Programs/CrashReportClient/
 ExcludePathFilters=Config/Localization/*
 ExcludePathFilters=Source/Runtime/Online/BuildPatchServices/*
 ExcludePathFilters=Plugins/*Editor*/*
+ExcludePathFilters=Plugins/MetaHuman/*
 ExcludePathFilters=Plugins/Online/*
 ExcludePathFilters=Plugins/Runtime/HTTPChunkInstaller/*
 ExcludePathFilters=*/NoRedist/*
diff --git a/Engine/Config/Localization/Keywords.ini b/Engine/Config/Localization/Keywords.ini
--- a/Engine/Config/Localization/Keywords.ini
+++ b/Engine/Config/Localization/Keywords.ini
@@ -27,6 +27,7 @@ IncludePathFilters=Source/Editor/*
 IncludePathFilters=Source/Runtime/*
 IncludePathFilters=Source/Developer/*
 IncludePathFilters=Plugins/*
+ExcludePathFilters=Plugins/MetaHuman/*
 ExcludePathFilters=*/NoRedist/*
 ExcludePathFilters=*/NotForLicensees/*
 ExcludePathFilters=*/LimitedAccess/*
diff --git a/Engine/Config/Localization/PropertyNames.ini b/Engine/Config/Localization/PropertyNames.ini
--- a/Engine/Config/Localization/PropertyNames.ini
+++ b/Engine/Config/Localization/PropertyNames.ini
@@ -32,6 +32,7 @@ IncludePathFilters=Source/Editor/*
 IncludePathFilters=Source/Runtime/*
 IncludePathFilters=Source/Developer/*
 IncludePathFilters=Plugins/*
+ExcludePathFilters=Plugins/MetaHuman/*
 ExcludePathFilters=*/NoRedist/*
 ExcludePathFilters=*/NotForLicensees/*
 ExcludePathFilters=*/LimitedAccess/*
diff --git a/Engine/Config/Localization/ToolTips.ini b/Engine/Config/Localization/ToolTips.ini
--- a/Engine/Config/Localization/ToolTips.ini
+++ b/Engine/Config/Localization/ToolTips.ini
@@ -32,6 +32,7 @@ IncludePathFilters=Source/Editor/*
 IncludePathFilters=Source/Runtime/*
 IncludePathFilters=Source/Developer/*
 IncludePathFilters=Plugins/*
+ExcludePathFilters=Plugins/MetaHuman/*
 ExcludePathFilters=*/NoRedist/*
 ExcludePathFilters=*/NotForLicensees/*
 ExcludePathFilters=*/LimitedAccess/*
