Temporarily gather MH category strings into the Engine loc files

#tests Run Localize command, observe UObjectCategories strings appeared in the corresponding files under 'Category' folder

diff --git a/Engine/Config/Localization/Category.ini b/Engine/Config/Localization/Category.ini
--- a/Engine/Config/Localization/Category.ini
+++ b/Engine/Config/Localization/Category.ini
@@ -32,7 +32,6 @@ IncludePathFilters=Source/Editor/*
 IncludePathFilters=Source/Runtime/*
 IncludePathFilters=Source/Developer/*
 IncludePathFilters=Plugins/*
-ExcludePathFilters=Plugins/MetaHuman/*
 ExcludePathFilters=*/NoRedist/*
 ExcludePathFilters=*/NotForLicensees/*
 ExcludePathFilters=*/LimitedAccess/*
