Fix missing localization for some MH Category strings

#tests Run Localize command locally with Engine as a working project and 'Category' as localization target. Observe new strings appeared in loc files under Engine/Content/Localization/Category folder.

diff --git a/Engine/Config/Localization/Category.ini b/Engine/Config/Localization/Category.ini
--- a/Engine/Config/Localization/Category.ini
+++ b/Engine/Config/Localization/Category.ini
@@ -35,7 +35,6 @@ IncludePathFilters=Plugins/*
 ExcludePathFilters=*/NoRedist/*
 ExcludePathFilters=*/NotForLicensees/*
 ExcludePathFilters=*/LimitedAccess/*
-ExcludePathFilters=Plugins/*/MetaHuman*
 InputKeys=Category
 OutputNamespaces=UObjectCategory
 OutputKeys="{MetaDataValue}"
