Filter PropertyNames strings gather for Engine localization.
Requires https://p4-swarm.epicgames.net/reviews/47545671

#tests Run Localize command locally, observe enums, funcions, BP nodes and pins are not gathered for localization anymore

diff --git a/Engine/Config/Localization/PropertyNames.ini b/Engine/Config/Localization/PropertyNames.ini
--- a/Engine/Config/Localization/PropertyNames.ini
+++ b/Engine/Config/Localization/PropertyNames.ini
@@ -35,6 +35,10 @@ IncludePathFilters=Plugins/*
 ExcludePathFilters=*/NoRedist/*
 ExcludePathFilters=*/NotForLicensees/*
 ExcludePathFilters=*/LimitedAccess/*
+FieldTypesToExclude=Enum
+FieldTypesToExclude=Function
+FieldOwnerTypesToExclude=K2Node
+FieldOuterTypesToExclude=Function
 InputKeys=DisplayName
 OutputNamespaces=UObjectDisplayNames
 OutputKeys="{FieldPath}"
